Moved the core functionality of the Plugin Browser for creating and mounting a new plugin into a new PluginUtils plugin.

diff --git a/Engine/Plugins/Developer/PluginUtils/PluginUtils.uplugin b/Engine/Plugins/Developer/PluginUtils/PluginUtils.uplugin
new file mode 100644
--- /dev/null
+++ b/Engine/Plugins/Developer/PluginUtils/PluginUtils.uplugin
@@ -0,0 +1,24 @@
+{
+	"FileVersion": 3,
+	"Version": 1,
+	"VersionName": "1.0",
+	"FriendlyName": "Plugin Utilities",
+	"Description": "Helpers to create and edit plugins. Used by Plugin Browser.",
+	"Category": "Programming",
+	"CreatedBy": "Epic Games, Inc.",
+	"CreatedByURL": "http://epicgames.com",
+	"DocsURL": "",
+	"MarketplaceURL": "",
+	"SupportURL": "",
+	"EnabledByDefault": true,
+	"CanContainContent": false,
+	"IsBetaVersion": false,
+	"Installed": false,
+	"Modules": [
+		{
+			"Name": "PluginUtils",
+			"Type": "Editor",
+			"LoadingPhase": "Default"
+		}
+	]
+}
diff --git a/Engine/Plugins/Developer/PluginUtils/Source/PluginUtils/PluginUtils.Build.cs b/Engine/Plugins/Developer/PluginUtils/Source/PluginUtils/PluginUtils.Build.cs
new file mode 100644
--- /dev/null
+++ b/Engine/Plugins/Developer/PluginUtils/Source/PluginUtils/PluginUtils.Build.cs
@@ -0,0 +1,27 @@
+// Copyright Epic Games, Inc. All Rights Reserved.
+
+namespace UnrealBuildTool.Rules
+{
+	public class PluginUtils : ModuleRules
+	{
+		public PluginUtils(ReadOnlyTargetRules Target) : base(Target)
+		{
+			PublicDependencyModuleNames.AddRange(
+				new string[] {
+					"Projects",
+					"Core",
+				}
+			);
+
+			PrivateDependencyModuleNames.AddRange(
+				new string[] {
+					"GameProjectGeneration",
+					"DesktopPlatform",
+					"AssetRegistry",
+					"AssetTools",
+					"CoreUObject",
+				}
+			);
+		}
+	}
+}
diff --git a/Engine/Plugins/Developer/PluginUtils/Source/PluginUtils/Private/PluginUtils.cpp b/Engine/Plugins/Developer/PluginUtils/Source/PluginUtils/Private/PluginUtils.cpp
new file mode 100644
--- /dev/null
+++ b/Engine/Plugins/Developer/PluginUtils/Source/PluginUtils/Private/PluginUtils.cpp
@@ -0,0 +1,484 @@
+// Copyright Epic Games, Inc. All Rights Reserved.
+
+#include "PluginUtils.h"
+
+#include "GameProjectUtils.h"
+#include "Interfaces/IProjectManager.h"
+#include "Interfaces/IPluginManager.h"
+#include "PluginDescriptor.h"
+#include "IAssetRegistry.h"
+#include "AssetRegistryModule.h"
+#include "AssetData.h"
+#include "IAssetTools.h"
+#include "AssetToolsModule.h"
+#include "DesktopPlatformModule.h"
+#include "HAL/FileManager.h"
+#include "HAL/PlatformFilemanager.h"
+#include "Misc/Paths.h"
+#include "Misc/FileHelper.h"
+#include "Misc/PackageName.h"
+#include "Misc/FeedbackContext.h"
+
+#define LOCTEXT_NAMESPACE "PluginUtils"
+
+namespace PluginUtils
+{
+	// The text macro to replace with the actual plugin name when copying files
+	const FString PLUGIN_NAME = TEXT("PLUGIN_NAME");
+
+	bool CopyPluginTemplateFolder(const TCHAR* DestinationDirectory, const TCHAR* Source, const FString& PluginName, FText& FailReason)
+	{
+		check(DestinationDirectory);
+		check(Source);
+
+		IPlatformFile& PlatformFile = FPlatformFileManager::Get().GetPlatformFile();
+
+		FString DestDir(DestinationDirectory);
+		FPaths::NormalizeDirectoryName(DestDir);
+
+		FString SourceDir(Source);
+		FPaths::NormalizeDirectoryName(SourceDir);
+
+		// Does Source dir exist?
+		if (!PlatformFile.DirectoryExists(*SourceDir))
+		{
+			FailReason = FText::Format(LOCTEXT("InvalidPluginTemplateFolder", "Plugin template folder '{0}' doesn't exist"), FText::FromString(FPaths::ConvertRelativePathToFull(SourceDir)));
+			return false;
+		}
+
+		// Destination directory exists already or can be created ?
+		if (!PlatformFile.DirectoryExists(*DestDir) && !PlatformFile.CreateDirectoryTree(*DestDir))
+		{
+			FailReason = FText::Format(LOCTEXT("FailedToCreateDestinationFolder", "Failed to create destination folder '{0}'"), FText::FromString(FPaths::ConvertRelativePathToFull(DestDir)));
+			return false;
+		}
+
+		// Copy all files and directories, renaming specific sections to the plugin name
+		struct FCopyPluginFilesAndDirs : public IPlatformFile::FDirectoryVisitor
+		{
+			IPlatformFile& PlatformFile;
+			const TCHAR* SourceRoot;
+			const TCHAR* DestRoot;
+			const FString& PluginName;
+			TArray<FString> NameReplacementFileTypes;
+			TArray<FString> IgnoredFileTypes;
+			TArray<FString> CopyUnmodifiedFileTypes;
+			FText& FailReason;
+
+			FCopyPluginFilesAndDirs(IPlatformFile& InPlatformFile, const TCHAR* InSourceRoot, const TCHAR* InDestRoot, const FString& InPluginName, FText& InFailReason)
+				: PlatformFile(InPlatformFile)
+				, SourceRoot(InSourceRoot)
+				, DestRoot(InDestRoot)
+				, PluginName(InPluginName)
+				, FailReason(InFailReason)
+			{
+				// Which file types we want to replace instances of PLUGIN_NAME with the new Plugin Name
+				NameReplacementFileTypes.Add(TEXT("cs"));
+				NameReplacementFileTypes.Add(TEXT("cpp"));
+				NameReplacementFileTypes.Add(TEXT("h"));
+				NameReplacementFileTypes.Add(TEXT("vcxproj"));
+
+				// Which file types do we want to ignore
+				IgnoredFileTypes.Add(TEXT("opensdf"));
+				IgnoredFileTypes.Add(TEXT("sdf"));
+				IgnoredFileTypes.Add(TEXT("user"));
+				IgnoredFileTypes.Add(TEXT("suo"));
+
+				// Which file types do we want to copy completely unmodified
+				CopyUnmodifiedFileTypes.Add(TEXT("uasset"));
+				CopyUnmodifiedFileTypes.Add(TEXT("umap"));
+			}
+
+			virtual bool Visit(const TCHAR* FilenameOrDirectory, bool bIsDirectory)
+			{
+				FString NewName(FilenameOrDirectory);
+				// change the root and rename paths/files
+				NewName.RemoveFromStart(SourceRoot);
+				NewName = NewName.Replace(*PLUGIN_NAME, *PluginName, ESearchCase::CaseSensitive);
+				NewName = FPaths::Combine(DestRoot, *NewName);
+
+				if (bIsDirectory)
+				{
+					// create new directory structure
+					if (!PlatformFile.CreateDirectoryTree(*NewName) && !PlatformFile.DirectoryExists(*NewName))
+					{
+						FailReason = FText::Format(LOCTEXT("FailedToCreatePluginSubFolder", "Failed to create plugin subfolder '{0}'"), FText::FromString(FPaths::ConvertRelativePathToFull(NewName)));
+						return false;
+					}
+				}
+				else
+				{
+					FString NewExt = FPaths::GetExtension(FilenameOrDirectory);
+
+					if (!IgnoredFileTypes.Contains(NewExt))
+					{
+						if (CopyUnmodifiedFileTypes.Contains(NewExt))
+						{
+							// Copy unmodified files with the original name, but do rename the directories
+							FString CleanFilename = FPaths::GetCleanFilename(FilenameOrDirectory);
+							FString CopyToPath = FPaths::GetPath(NewName);
+
+							NewName = FPaths::Combine(CopyToPath, CleanFilename);
+						}
+
+						if (PlatformFile.FileExists(*NewName))
+						{
+							// Delete destination file if it exists
+							PlatformFile.DeleteFile(*NewName);
+						}
+
+						// If file of specified extension - open the file as text and replace PLUGIN_NAME in there before saving
+						if (NameReplacementFileTypes.Contains(NewExt))
+						{
+							FString OutFileContents;
+							if (!FFileHelper::LoadFileToString(OutFileContents, FilenameOrDirectory))
+							{
+								FailReason = FText::Format(LOCTEXT("FailedToReadPluginTemplateFile", "Failed to read plugin template file '{0}'"), FText::FromString(FPaths::ConvertRelativePathToFull(FilenameOrDirectory)));
+								return false;
+							}
+
+							OutFileContents = OutFileContents.Replace(*PLUGIN_NAME, *PluginName, ESearchCase::CaseSensitive);
+
+							// For some content, we also want to export a PLUGIN_NAME_API text macro, which requires that the plugin name
+							// be all capitalized
+
+							FString PluginNameAPI = PluginName + TEXT("_API");
+
+							OutFileContents = OutFileContents.Replace(*PluginNameAPI, *PluginNameAPI.ToUpper(), ESearchCase::CaseSensitive);
+
+							if (!FFileHelper::SaveStringToFile(OutFileContents, *NewName))
+							{
+								FailReason = FText::Format(LOCTEXT("FailedToWritePluginFile", "Failed to write plugin file '{0}'"), FText::FromString(FPaths::ConvertRelativePathToFull(NewName)));
+								return false;
+							}
+						}
+						else
+						{
+							// Copy file from source
+							if (!PlatformFile.CopyFile(*NewName, FilenameOrDirectory))
+							{
+								// Not all files could be copied
+								FailReason = FText::Format(LOCTEXT("FailedToCopyPluginTemplateFile", "Failed to copy plugin template file '{0}' to '{1}'"), FText::FromString(FPaths::ConvertRelativePathToFull(FilenameOrDirectory)), FText::FromString(FPaths::ConvertRelativePathToFull(NewName)));
+								return false;
+							}
+						}
+					}
+				}
+				return true; // continue searching
+			}
+		};
+
+		// copy plugin files and directories visitor
+		FCopyPluginFilesAndDirs CopyFilesAndDirs(PlatformFile, *SourceDir, *DestDir, PluginName, FailReason);
+
+		// create all files subdirectories and files in subdirectories!
+		return PlatformFile.IterateDirectoryRecursively(*SourceDir, CopyFilesAndDirs);
+	}
+
+	void FixupPluginTemplateAssets(const FString& PluginName)
+	{
+		IPlatformFile& PlatformFile = FPlatformFileManager::Get().GetPlatformFile();
+
+		TSharedPtr<IPlugin> Plugin = IPluginManager::Get().FindPlugin(PluginName);
+
+		struct FFixupPluginAssets : public IPlatformFile::FDirectoryVisitor
+		{
+			IPlatformFile& PlatformFile;
+			const FString& PluginName;
+			const FString& PluginBaseDir;
+
+			TArray<FString> FilesToScan;
+
+			FFixupPluginAssets(IPlatformFile& InPlatformFile, const FString& InPluginName, const FString& InPluginBaseDir)
+				: PlatformFile(InPlatformFile)
+				, PluginName(InPluginName)
+				, PluginBaseDir(InPluginBaseDir)
+			{
+			}
+
+			virtual bool Visit(const TCHAR* FilenameOrDirectory, bool bIsDirectory) override
+			{
+				if (!bIsDirectory)
+				{
+					FString Extension = FPaths::GetExtension(FilenameOrDirectory);
+
+					// Only interested in fixing up uassets and umaps...anything else we leave alone
+					if (Extension == TEXT("uasset") || Extension == TEXT("umap"))
+					{
+						FilesToScan.Add(FilenameOrDirectory);
+					}
+				}
+
+				return true;
+			}
+
+			/**
+			 * Fixes up any assets that contain the PLUGIN_NAME text macro, since those need to be renamed by the engine for the change to
+			 * stick (as opposed to just renaming the file)
+			 */
+			void PerformFixup()
+			{
+				TArray<FAssetRenameData> AssetRenameData;
+
+				if (FilesToScan.Num() > 0)
+				{
+					IAssetRegistry& AssetRegistry = FModuleManager::LoadModuleChecked<FAssetRegistryModule>("AssetRegistry").Get();
+					AssetRegistry.ScanFilesSynchronous(FilesToScan);
+
+					for (const FString& File : FilesToScan)
+					{
+						TArray<FAssetData> Assets;
+
+						FString PackageName;
+						if (FPackageName::TryConvertFilenameToLongPackageName(File, PackageName))
+						{
+							AssetRegistry.GetAssetsByPackageName(*PackageName, Assets);
+						}
+
+						for (FAssetData Asset : Assets)
+						{
+							FString AssetName = Asset.AssetName.ToString().Replace(*PLUGIN_NAME, *PluginName, ESearchCase::CaseSensitive);
+							FString AssetPath = Asset.PackagePath.ToString().Replace(*PLUGIN_NAME, *PluginName, ESearchCase::CaseSensitive);
+
+							FAssetRenameData RenameData(Asset.GetAsset(), AssetPath, AssetName);
+
+							AssetRenameData.Add(RenameData);
+						}
+					}
+
+					if (AssetRenameData.Num() > 0)
+					{
+						FAssetToolsModule& AssetToolsModule = FModuleManager::LoadModuleChecked<FAssetToolsModule>("AssetTools");
+						AssetToolsModule.Get().RenameAssetsWithDialog(AssetRenameData);
+					}
+				}
+			}
+		};
+
+		if (Plugin.IsValid())
+		{
+			const FString PluginBaseDir = Plugin->GetBaseDir();
+			FFixupPluginAssets FixupPluginAssets(PlatformFile, PluginName, PluginBaseDir);
+			PlatformFile.IterateDirectoryRecursively(*PluginBaseDir, FixupPluginAssets);
+			FixupPluginAssets.PerformFixup();
+		}
+	}
+}
+
+FString FPluginUtils::GetPluginFolder(const FString& PluginLocation, const FString& PluginName, bool bFullPath)
+{
+	FString PluginFolder = FPaths::Combine(PluginLocation, PluginName);
+	if (bFullPath)
+	{
+		PluginFolder = FPaths::ConvertRelativePathToFull(PluginFolder);
+	}
+	FPaths::MakePlatformFilename(PluginFolder);
+	return PluginFolder;
+}
+
+FString FPluginUtils::GetPluginFilePath(const FString& PluginLocation, const FString& PluginName, bool bFullPath)
+{
+	FString PluginFilePath = FPaths::Combine(PluginLocation, PluginName, (PluginName + TEXT(".uplugin")));
+	if (bFullPath)
+	{
+		PluginFilePath = FPaths::ConvertRelativePathToFull(PluginFilePath);
+	}
+	FPaths::MakePlatformFilename(PluginFilePath);
+	return PluginFilePath;
+}
+
+FString FPluginUtils::GetPluginContentFolder(const FString& PluginLocation, const FString& PluginName, bool bFullPath)
+{
+	FString PluginContentFolder = FPaths::Combine(PluginLocation, PluginName, TEXT("Content"));
+	if (bFullPath)
+	{
+		PluginContentFolder = FPaths::ConvertRelativePathToFull(PluginContentFolder);
+	}
+	FPaths::MakePlatformFilename(PluginContentFolder);
+	return PluginContentFolder;
+}
+
+FString FPluginUtils::GetPluginResourcesFolder(const FString& PluginLocation, const FString& PluginName, bool bFullPath)
+{
+	FString PluginResourcesFolder = FPaths::Combine(PluginLocation, PluginName, TEXT("Resources"));
+	if (bFullPath)
+	{
+		PluginResourcesFolder = FPaths::ConvertRelativePathToFull(PluginResourcesFolder);
+	}
+	FPaths::MakePlatformFilename(PluginResourcesFolder);
+	return PluginResourcesFolder;
+}
+
+bool FPluginUtils::CreateAndMountNewPlugin(const FNewPluginParams& Params, FText& FailReason)
+{
+	// Early validations on new plugin params
+	if (Params.PluginName.IsEmpty())
+	{
+		FailReason = LOCTEXT("CreateNewPluginParam_NoPluginName", "Missing plugin name");
+		return false;
+	}
+	if (Params.PluginLocation.IsEmpty())
+	{
+		FailReason = LOCTEXT("CreateNewPluginParam_NoPluginLocation", "Missing plugin location");
+		return false;
+	}
+	if (Params.bHasModules && Params.TemplateFolders.Num() == 0)
+	{
+		FailReason = LOCTEXT("CreateNewPluginParam_NoTemplateFolder", "A template folder must be specified to create a plugin with code");
+		return false;
+	}
+
+	const FString PluginFolder = FPluginUtils::GetPluginFolder(Params.PluginLocation, Params.PluginName, /*bFullPath*/ true);
+
+	bool bSucceeded = true;
+	do 
+	{
+		IPlatformFile& PlatformFile = FPlatformFileManager::Get().GetPlatformFile();
+		if (!PlatformFile.DirectoryExists(*PluginFolder) && !PlatformFile.CreateDirectoryTree(*PluginFolder))
+		{
+			FailReason = FText::Format(LOCTEXT("FailedToCreatePluginFolder", "Failed to create plugin folder '{0}'"), FText::FromString(PluginFolder));
+			bSucceeded = false;
+			break;
+		}
+
+		if (Params.bCanContainContent)
+		{
+			const FString PluginContentFolder = FPluginUtils::GetPluginContentFolder(Params.PluginLocation, Params.PluginName, /*bFullPath*/ true);
+			if (!PlatformFile.DirectoryExists(*PluginContentFolder) && !PlatformFile.CreateDirectory(*PluginContentFolder))
+			{
+				FailReason = FText::Format(LOCTEXT("FailedToCreatePluginContentFolder", "Failed to create plugin Content folder '{0}'"), FText::FromString(PluginContentFolder));
+				bSucceeded = false;
+				break;
+			}
+		}
+
+		FPluginDescriptor Descriptor;
+		Descriptor.FriendlyName = Params.PluginName;
+		Descriptor.Version = 1;
+		Descriptor.VersionName = TEXT("1.0");
+		Descriptor.Category = TEXT("Other");
+		Descriptor.CreatedBy = Params.CreatedBy;
+		Descriptor.CreatedByURL = Params.CreatedByURL;
+		Descriptor.Description = Params.Description;
+		Descriptor.bIsBetaVersion = Params.bIsBetaVersion;
+		Descriptor.bCanContainContent = Params.bCanContainContent;
+		if (Params.bHasModules)
+		{
+			Descriptor.Modules.Add(FModuleDescriptor(*Params.PluginName, Params.ModuleDescriptorType, Params.LoadingPhase));
+		}
+
+		// Write the uplugin file
+		const FString PluginFilePath = FPluginUtils::GetPluginFilePath(Params.PluginLocation, Params.PluginName, /*bFullPath*/ true);
+		if (!Descriptor.Save(PluginFilePath, FailReason))
+		{
+			bSucceeded = false;
+			break;
+		}
+
+		if (!Params.PluginIconPath.IsEmpty())
+		{
+			const FString ResourcesFolder = FPluginUtils::GetPluginResourcesFolder(Params.PluginLocation, Params.PluginName, /*bFullPath*/ true);
+			const FString DestinationPluginIconPath = FPaths::Combine(ResourcesFolder, TEXT("Icon128.png"));
+			if (IFileManager::Get().Copy(*DestinationPluginIconPath, *Params.PluginIconPath, /*bReplaceExisting=*/ false) != COPY_OK)
+			{
+				FailReason = FText::Format(LOCTEXT("FailedToCopyPluginIcon", "Failed to copy plugin icon from '{0}' to '{1}'"), FText::FromString(FPaths::ConvertRelativePathToFull(Params.PluginIconPath)), FText::FromString(DestinationPluginIconPath));
+				bSucceeded = false;
+				break;
+			}
+		}
+
+		GWarn->BeginSlowTask(LOCTEXT("CopyingPluginTemplate", "Copying plugin template files..."), /*ShowProgressDialog*/ true, /*bShowCancelButton*/ false);
+		for (const FString& TemplateFolder : Params.TemplateFolders)
+		{
+			if (!PluginUtils::CopyPluginTemplateFolder(*PluginFolder, *TemplateFolder, Params.PluginName, FailReason))
+			{
+				FailReason = FText::Format(LOCTEXT("FailedToCopyPluginTemplate", "Failed to copy plugin template files from '{0}' to '{1}'. {2}"), FText::FromString(FPaths::ConvertRelativePathToFull(TemplateFolder)), FText::FromString(PluginFolder), FailReason);
+				bSucceeded = false;
+				break;
+			}
+		}
+		GWarn->EndSlowTask();
+		if (!bSucceeded)
+		{
+			break;
+		}
+
+		if (Params.bHasModules)
+		{
+			const FString ProjectFileName = IFileManager::Get().ConvertToAbsolutePathForExternalAppForWrite(*FPaths::GetProjectFilePath());
+			const FString Arguments = FString::Printf(TEXT("%s %s %s -Plugin=\"%s\" -Project=\"%s\" -Progress -NoHotReloadFromIDE"), FPlatformMisc::GetUBTTargetName(), FModuleManager::Get().GetUBTConfiguration(), FPlatformMisc::GetUBTPlatform(), *PluginFilePath, *ProjectFileName);
+			if (!FDesktopPlatformModule::Get()->RunUnrealBuildTool(FText::Format(LOCTEXT("CompilingPlugin", "Compiling {0} plugin..."), FText::FromString(Params.PluginName)), FPaths::RootDir(), Arguments, GWarn))
+			{
+				FailReason = LOCTEXT("FailedToCompilePlugin", "Failed to compile plugin source code");
+				bSucceeded = false;
+				break;
+			}
+
+			// Reset the module paths cache. For unique build environments, the modules may be generated to the project binaries directory.
+			FModuleManager::Get().ResetModulePathsCache();
+
+			// Generate project files if we happen to be using a project file.
+			if (!FDesktopPlatformModule::Get()->GenerateProjectFiles(FPaths::RootDir(), FPaths::GetProjectFilePath(), GWarn))
+			{
+				FailReason = LOCTEXT("FailedToGenerateProjectFiles", "Failed to generate project files");
+				bSucceeded = false;
+				break;
+			}
+		}
+
+		if (Params.bUpdateProjectAddtitionalPluginDirectories)
+		{
+			// If this path isn't in the Engine/Plugins dir and isn't in Project/Plugins dir,
+			// add the directory to the list of ones we additionally scan.
+			FString PluginLocationFull = FPaths::ConvertRelativePathToFull(Params.PluginLocation);
+			FPaths::MakePlatformFilename(PluginLocationFull);
+			FString AbsoluteEnginePluginPath = IFileManager::Get().ConvertToAbsolutePathForExternalAppForRead(*FPaths::EnginePluginsDir());
+			FPaths::MakePlatformFilename(AbsoluteEnginePluginPath);
+			if (!PluginLocationFull.StartsWith(AbsoluteEnginePluginPath))
+			{
+				// There have been issues with ProjectDir can be relative and PluginLocation absolute, causing our
+				// tests to fail below. We now normalize on absolute paths prior to performing the check to ensure
+				// that we don't add the folder to the additional plugin directory unnecessarily (which can result in build failures).
+				FString ProjectDirFull = FPaths::ConvertRelativePathToFull(FPaths::ProjectDir());
+				FPaths::MakePlatformFilename(ProjectDirFull);
+				if (!PluginLocationFull.StartsWith(ProjectDirFull))
+				{
+					GameProjectUtils::UpdateAdditionalPluginDirectory(PluginLocationFull, true);
+				}
+			}
+		}
+
+		// Update the list of known plugins
+		IPluginManager::Get().RefreshPluginsList();
+
+		// Enable this plugin in the project
+		if (!IProjectManager::Get().SetPluginEnabled(Params.PluginName, true, FailReason))
+		{
+			FailReason = FText::Format(LOCTEXT("FailedToEnablePlugin", "Failed to enable plugin: {0}"), FailReason);
+			bSucceeded = false;
+			break;
+		}
+
+		// Mount the plugin
+		GWarn->BeginSlowTask(LOCTEXT("MountingPluginFiles", "Mounting plugin files..."), /*ShowProgressDialog*/ true, /*bShowCancelButton*/ false);
+		IPluginManager::Get().MountNewlyCreatedPlugin(Params.PluginName);
+		GWarn->EndSlowTask();
+
+		if (Params.bCanContainContent)
+		{	
+			// Attempt to fix any content that was added by the plugin
+			GWarn->BeginSlowTask(LOCTEXT("LoadingContent", "Loading Content..."), /*ShowProgressDialog*/ true, /*bShowCancelButton*/ false);
+			PluginUtils::FixupPluginTemplateAssets(Params.PluginName);
+			GWarn->EndSlowTask();
+		}
+	} while (false);
+
+	if (!bSucceeded)
+	{
+		IFileManager::Get().DeleteDirectory(*PluginFolder, /*RequireExists*/ false, /*Tree*/ true);
+	}
+
+	return bSucceeded;
+}
+
+#undef LOCTEXT_NAMESPACE
diff --git a/Engine/Plugins/Developer/PluginUtils/Source/PluginUtils/Private/PluginUtilsModule.cpp b/Engine/Plugins/Developer/PluginUtils/Source/PluginUtils/Private/PluginUtilsModule.cpp
new file mode 100644
--- /dev/null
+++ b/Engine/Plugins/Developer/PluginUtils/Source/PluginUtils/Private/PluginUtilsModule.cpp
@@ -0,0 +1,25 @@
+// Copyright Epic Games, Inc. All Rights Reserved.
+
+#include "Modules/ModuleInterface.h"
+#include "Modules/ModuleManager.h"
+
+/**
+ * Implements the PluginUtils module.
+ */
+class FPluginUtilsModule
+	: public IModuleInterface
+{
+public:
+
+	virtual void StartupModule( ) override { }
+
+	virtual void ShutdownModule( ) override { }
+
+	virtual bool SupportsDynamicReloading( ) override
+	{
+		return true;
+	}
+};
+
+
+IMPLEMENT_MODULE(FPluginUtilsModule, PluginUtils);
diff --git a/Engine/Plugins/Developer/PluginUtils/Source/PluginUtils/Public/PluginUtils.h b/Engine/Plugins/Developer/PluginUtils/Source/PluginUtils/Public/PluginUtils.h
new file mode 100644
--- /dev/null
+++ b/Engine/Plugins/Developer/PluginUtils/Source/PluginUtils/Public/PluginUtils.h
@@ -0,0 +1,105 @@
+// Copyright Epic Games, Inc. All Rights Reserved.
+
+#pragma once
+
+#include "CoreMinimal.h"
+#include "ModuleDescriptor.h"
+
+class PLUGINUTILS_API FPluginUtils
+{
+public:
+	/**
+	 * Returns the plugin folder.
+	 * @param PluginLocation directory that contains the plugin folder
+	 * @param PluginName name of the plugin
+	 * @param bFullPath ensures a full path is returned
+	 */
+	static FString GetPluginFolder(const FString& PluginLocation, const FString& PluginName, bool bFullPath = false);
+
+	/** 
+	 * Returns the uplugin file path.
+	 * @param PluginLocation directory that contains the plugin folder
+	 * @param PluginName name of the plugin
+	 * @param bFullPath ensures a full path is returned
+	 */
+	static FString GetPluginFilePath(const FString& PluginLocation, const FString& PluginName, bool bFullPath = false);
+
+	/**
+	 * Returns the plugin Content folder.
+	 * @param PluginLocation directory that contains the plugin folder
+	 * @param PluginName name of the plugin
+	 * @param bFullPath ensures a full path is returned
+	 */
+	static FString GetPluginContentFolder(const FString& PluginLocation, const FString& PluginName, bool bFullPath = false);
+
+	/**
+	 * Returns the plugin Resources folder.
+	 * @param PluginLocation directory that contains the plugin folder
+	 * @param PluginName name of the plugin
+	 * @param bFullPath ensures a full path is returned
+	 */
+	static FString GetPluginResourcesFolder(const FString& PluginLocation, const FString& PluginName, bool bFullPath = false);
+
+	/**
+	 * Parameters for creating a new plugin.
+	 */
+	struct FNewPluginParams
+	{
+		FNewPluginParams(const FString& InPluginName, const FString& InPluginLocation)
+			: PluginName(InPluginName)
+			, PluginLocation(InPluginLocation)
+		{}
+
+		/** Plugin name (Required) */
+		FString PluginName;
+		
+		/** Directory in which to create the plugin folder (Required) */
+		FString PluginLocation;
+
+		/** The author of this plugin */
+		FString CreatedBy;
+
+		/** Hyperlink for the author's website  */
+		FString CreatedByURL;
+
+		/** A description for this plugin */
+		FString Description;
+
+		/** Path to plugin icon to copy in the plugin resources folder */
+		FString PluginIconPath;
+
+		/** 
+		 * Folders containing template files to copy into the plugin folder (Required if bHasModules).
+		 * Occurrences of the string PLUGIN_NAME in the filename or file content will be replaced by the plugin name. 
+		 */
+		TArray<FString> TemplateFolders;
+
+		/** Marks this content as being in beta */
+		bool bIsBetaVersion = false;
+
+		/** Can this plugin contain content */
+		bool bCanContainContent = false;
+
+		/** Does this plugin have Source files? */
+		bool bHasModules = false;
+
+		/**
+		 * Whether to update the project additional plugin directories (persistently saved in uproject file)
+		 * if the specified plugin location is outside the engine or game plugin folders. 
+		 */
+		bool bUpdateProjectAddtitionalPluginDirectories = true;
+
+		/** If this plugin has Source, what is the type of Source included (so it can potentially be excluded in the right builds) */
+		EHostType::Type ModuleDescriptorType = EHostType::Runtime;
+
+		/** If this plugin has Source, when should the module be loaded (may need to be earlier than default if used in blueprints) */
+		ELoadingPhase::Type LoadingPhase = ELoadingPhase::Default;
+	};
+
+	/**
+	 * Helper to create and mount a new unreal plugin.
+	 * @param FailReason Reason the plugin creation failed.
+	 * @return Whether the plugin was successfully created. If something goes wrong during the creation process, the plugin folder gets deleted.
+	 */
+	static bool CreateAndMountNewPlugin(const FNewPluginParams& Params, FText& FailReason);
+};
diff --git a/Engine/Plugins/Editor/PluginBrowser/PluginBrowser.uplugin b/Engine/Plugins/Editor/PluginBrowser/PluginBrowser.uplugin
--- a/Engine/Plugins/Editor/PluginBrowser/PluginBrowser.uplugin
+++ b/Engine/Plugins/Editor/PluginBrowser/PluginBrowser.uplugin
@@ -21,5 +21,11 @@
 			"Type" : "Editor",
 			"LoadingPhase" : "Default"
 		}
+	],
+	"Plugins": [ 
+		{ 
+			"Name": "PluginUtils",
+			"Enabled": true 
+		} 
 	]
 }
\ No newline at end of file
diff --git a/Engine/Plugins/Editor/PluginBrowser/Source/PluginBrowser/PluginBrowser.Build.cs b/Engine/Plugins/Editor/PluginBrowser/Source/PluginBrowser/PluginBrowser.Build.cs
--- a/Engine/Plugins/Editor/PluginBrowser/Source/PluginBrowser/PluginBrowser.Build.cs
+++ b/Engine/Plugins/Editor/PluginBrowser/Source/PluginBrowser/PluginBrowser.Build.cs
@@ -19,6 +19,7 @@ namespace UnrealBuildTool.Rules
 			
 			PrivateDependencyModuleNames.AddRange(
 				new string[] {
+					"PluginUtils",
 					"EditorStyle",
 					"Projects",
 					"UnrealEd",
diff --git a/Engine/Plugins/Editor/PluginBrowser/Source/PluginBrowser/Private/PluginHelpers.cpp b/Engine/Plugins/Editor/PluginBrowser/Source/PluginBrowser/Private/PluginHelpers.cpp
deleted file mode 100644
diff --git a/Engine/Plugins/Editor/PluginBrowser/Source/PluginBrowser/Private/PluginHelpers.h b/Engine/Plugins/Editor/PluginBrowser/Source/PluginBrowser/Private/PluginHelpers.h
deleted file mode 100644
diff --git a/Engine/Plugins/Editor/PluginBrowser/Source/PluginBrowser/Private/SNewPluginWizard.cpp b/Engine/Plugins/Editor/PluginBrowser/Source/PluginBrowser/Private/SNewPluginWizard.cpp
--- a/Engine/Plugins/Editor/PluginBrowser/Source/PluginBrowser/Private/SNewPluginWizard.cpp
+++ b/Engine/Plugins/Editor/PluginBrowser/Source/PluginBrowser/Private/SNewPluginWizard.cpp
@@ -4,7 +4,6 @@
 #include "Misc/Paths.h"
 #include "HAL/PlatformFilemanager.h"
 #include "Misc/MessageDialog.h"
-#include "Misc/FeedbackContext.h"
 #include "HAL/FileManager.h"
 #include "Misc/App.h"
 #include "Widgets/SBoxPanel.h"
@@ -26,7 +25,6 @@
 #include "Widgets/Views/SListView.h"
 #include "Widgets/Views/STileView.h"
 #include "PluginStyle.h"
-#include "PluginHelpers.h"
 #include "DesktopPlatformModule.h"
 #include "Widgets/Docking/SDockTab.h"
 #include "Framework/Notifications/NotificationManager.h"
@@ -37,8 +35,8 @@
 #include "SFilePathBlock.h"
 #include "Interfaces/IProjectManager.h"
 #include "ProjectDescriptor.h"
+#include "PluginUtils.h"
 #include "Interfaces/IMainFrameModule.h"
-#include "Interfaces/IProjectManager.h"
 #include "GameProjectGenerationModule.h"
 #include "DefaultPluginWizardDefinition.h"
 #include "UnrealEdMisc.h"
@@ -645,10 +643,7 @@ FString SNewPluginWizard::GetPluginFilenameWithPath() const
 	}
 	else
 	{
-		const FString& TestPluginName = PluginNameText.ToString();
-		FString TestPluginPath = PluginFolderPath / TestPluginName / (TestPluginName + TEXT(".uplugin"));
-		FPaths::MakePlatformFilename(TestPluginPath);
-		return TestPluginPath;
+		return FPluginUtils::GetPluginFilePath(PluginFolderPath, PluginNameText.ToString());
 	}
 }
 
@@ -678,137 +673,40 @@ void SNewPluginWizard::OnEnginePluginCheckboxChanged(ECheckBoxState NewCheckedSt
 
 FReply SNewPluginWizard::OnCreatePluginClicked()
 {
-	const FString AutoPluginName = PluginNameText.ToString();
-
-	// Plugin thumbnail image
-	FString PluginEditorIconPath;
-	bool bRequiresDefaultIcon = PluginWizardDefinition->GetPluginIconPath(PluginEditorIconPath);
-
-	TArray<FString> CreatedFiles;
-
-	FText LocalFailReason;
-
-	bool bSucceeded = true;
-
-	const bool bHasModules = PluginWizardDefinition->HasModules();
-
-	// Save descriptor file as .uplugin file
-	const FString UPluginFilePath = GetPluginFilenameWithPath();
-
-	// Define additional parameters to write out the plugin descriptor
-	FWriteDescriptorParams DescriptorParams;
+	if (!ensure(!PluginFolderPath.IsEmpty() && !PluginNameText.IsEmpty()))
 	{
-		DescriptorParams.bCanContainContent = PluginWizardDefinition->CanContainContent();
-		DescriptorParams.bHasModules = bHasModules;
-		DescriptorParams.ModuleDescriptorType = PluginWizardDefinition->GetPluginModuleDescriptor();
-		DescriptorParams.LoadingPhase = PluginWizardDefinition->GetPluginLoadingPhase();
-	}
-
-	const FString PluginModuleName = AutoPluginName;
-	bSucceeded = bSucceeded && WritePluginDescriptor(PluginModuleName, UPluginFilePath, DescriptorParams);
-
-	// Main plugin dir
-	const FString BasePluginFolder = GetPluginDestinationPath().ToString();
-	const FString PluginFolder = BasePluginFolder / AutoPluginName;
-
-	// Resource folder
-	const FString ResourcesFolder = PluginFolder / TEXT("Resources");
-
-	if (bRequiresDefaultIcon)
-	{
-		// Copy the icon
-		bSucceeded = bSucceeded && CopyFile(ResourcesFolder / TEXT("Icon128.png"), PluginEditorIconPath, /*inout*/ CreatedFiles);
+		// Don't even try to assemble the path or else it may be relative to the binaries folder!
+		return FReply::Unhandled();
 	}
 
-	TArray<FString> TemplateFolders = PluginWizardDefinition->GetFoldersForSelection();
-	if (TemplateFolders.Num() == 0)
-	{
-		PopErrorNotification(LOCTEXT("FailedTemplateCopy_NoFolders", "No templates were selected to create the plugin"));
-		bSucceeded = false;
-	}
+	bool bSucceeded = true;
+	FText FailReason;
 
-	GWarn->BeginSlowTask(LOCTEXT("CopyingData", "Copying data..."), true, false);
-	for (FString TemplateFolderName : TemplateFolders)
-	{
-		if (bSucceeded && !FPluginHelpers::CopyPluginTemplateFolder(*PluginFolder, *TemplateFolderName, AutoPluginName))
-		{
-			PopErrorNotification(FText::Format(LOCTEXT("FailedTemplateCopy", "Failed to copy plugin Template: {0}"), FText::FromString(TemplateFolderName)));
-			bSucceeded = false;
-			break;
-		}
-	}
-	GWarn->EndSlowTask();
+	const FString PluginName = PluginNameText.ToString();
+	const bool bHasModules = PluginWizardDefinition->HasModules();
 
-	// If it contains code, we need the user to restart to enable it. Otherwise, we can just mount it now.
-	if (bSucceeded && bHasModules)
+	FPluginUtils::FNewPluginParams NewPluginParams(PluginName, PluginFolderPath);
+	NewPluginParams.TemplateFolders = PluginWizardDefinition->GetFoldersForSelection();
+	NewPluginParams.bCanContainContent = PluginWizardDefinition->CanContainContent();
+	NewPluginParams.bHasModules = bHasModules;
+	NewPluginParams.ModuleDescriptorType = PluginWizardDefinition->GetPluginModuleDescriptor();
+	NewPluginParams.LoadingPhase = PluginWizardDefinition->GetPluginLoadingPhase();
+	PluginWizardDefinition->GetPluginIconPath(NewPluginParams.PluginIconPath);
+	if (DescriptorData.IsValid())
 	{
-		FString ProjectFileName = IFileManager::Get().ConvertToAbsolutePathForExternalAppForWrite(*FPaths::GetProjectFilePath());
-		FString Arguments = FString::Printf(TEXT("%s %s %s -Plugin=\"%s\" -Project=\"%s\" -Progress -NoHotReloadFromIDE"), FPlatformMisc::GetUBTTargetName(), FModuleManager::Get().GetUBTConfiguration(), FPlatformMisc::GetUBTPlatform(), *UPluginFilePath, *ProjectFileName);
-		if (!FDesktopPlatformModule::Get()->RunUnrealBuildTool(LOCTEXT("Compiling", "Compiling..."), FPaths::RootDir(), Arguments, GWarn))
-		{
-			PopErrorNotification(LOCTEXT("FailedToCompile", "Failed to compile source code."));
-			bSucceeded = false;
-		}
-
-		// Reset the module paths cache. For unique build environments, the modules may be generated to the project binaries directory.
-		FModuleManager::Get().ResetModulePathsCache();
-
-		// Generate project files if we happen to be using a project file.
-		if (bSucceeded && !FDesktopPlatformModule::Get()->GenerateProjectFiles(FPaths::RootDir(), FPaths::GetProjectFilePath(), GWarn))
-		{
-			PopErrorNotification(LOCTEXT("FailedToGenerateProjectFiles", "Failed to generate project files."));
-			bSucceeded = false;
-		}
-	}
-
-	if (bSucceeded)
-	{
-		// Notify that a new plugin has been created
-		FPluginBrowserModule& PluginBrowserModule = FPluginBrowserModule::Get();
-		PluginBrowserModule.BroadcastNewPluginCreated();
-
-		// Enable game plugins immediately
-		if (!bIsEnginePlugin)
-		{
-			// If this path isn't in the Engine/Plugins dir and isn't in Project/Plugins dir,
-			// add the directory to the list of ones we additionally scan
-			
-			// There have been issues with ProjectDir can be relative and BasePluginFolder absolute, causing our
-			// tests to fail below. We now normalize on absolute paths prior to performing the check to ensure
-			// that we don't add the folder to the additional plugin directory unnecessarily (which can result in build failures).
-			FString ProjectDirFull = FPaths::ConvertRelativePathToFull(FPaths::ProjectDir());
-			FString BasePluginFolderFull = FPaths::ConvertRelativePathToFull(BasePluginFolder);
-			if (!BasePluginFolderFull.StartsWith(ProjectDirFull))
-			{
-				GameProjectUtils::UpdateAdditionalPluginDirectory(BasePluginFolderFull, true);
-			}
-		}
-
-		// Update the list of known plugins
-		IPluginManager::Get().RefreshPluginsList();
-
-		// Enable this plugin in the project, if necessary
-		FText FailReason;
-		if (!IProjectManager::Get().SetPluginEnabled(AutoPluginName, true, FailReason))
-		{
-			PopErrorNotification(FText::Format(LOCTEXT("FailedToEnablePlugin", "Couldn't enable plugin: {0}"), FailReason));
-			bSucceeded = false;
-		}
-
-		// Mount the plugin
-		if (bSucceeded)
-		{
-			GWarn->BeginSlowTask(LOCTEXT("MountingFiles", "Mounting files..."), true, false);
-			IPluginManager::Get().MountNewlyCreatedPlugin(AutoPluginName);
-			GWarn->EndSlowTask();
-		}
+		NewPluginParams.CreatedBy = DescriptorData->CreatedBy;
+		NewPluginParams.CreatedByURL = DescriptorData->CreatedByURL;
+		NewPluginParams.Description = DescriptorData->Description;
+		NewPluginParams.bIsBetaVersion = DescriptorData->bIsBetaVersion;
 	}
+	
+	bSucceeded = bSucceeded && FPluginUtils::CreateAndMountNewPlugin(NewPluginParams, FailReason);
 
 	// Set the content browser to show the plugin's content directory
 	if (bSucceeded && PluginWizardDefinition->CanContainContent() && ShowPluginContentDirectoryCheckBox->IsChecked())
 	{
 		TArray<FString> SelectedDirectories;
-		SelectedDirectories.Add(TEXT("/") + AutoPluginName);
+		SelectedDirectories.Add(TEXT("/") + PluginName);
 
 		const bool bContentBrowserNeedsRefresh = true;
 
@@ -818,19 +716,15 @@ FReply SNewPluginWizard::OnCreatePluginClicked()
 		ContentBrowser.SetSelectedPaths(SelectedDirectories, bContentBrowserNeedsRefresh);
 	}
 
-	if (bSucceeded && PluginWizardDefinition->CanContainContent())
-	{
-		GWarn->BeginSlowTask(LOCTEXT("LoadingContent", "Loading Content..."), true, false);
-		// Attempt to fix any content that was added by the plugin
-		bSucceeded = FPluginHelpers::FixupPluginTemplateAssets(AutoPluginName);
-		GWarn->EndSlowTask();
-	}
+	PluginWizardDefinition->PluginCreated(PluginName, bSucceeded);
 
-	PluginWizardDefinition->PluginCreated(AutoPluginName, bSucceeded);
-	// Trigger the plugin manager to mount the new plugin, or delete the partially created plugin and abort
 	if (bSucceeded)
 	{
-		FNotificationInfo Info(FText::Format(LOCTEXT("PluginCreatedSuccessfully", "'{0}' was created successfully."), FText::FromString(AutoPluginName)));
+		// Notify that a new plugin has been created
+		FPluginBrowserModule& PluginBrowserModule = FPluginBrowserModule::Get();
+		PluginBrowserModule.BroadcastNewPluginCreated();
+
+		FNotificationInfo Info(FText::Format(LOCTEXT("PluginCreatedSuccessfully", "'{0}' was created successfully."), FText::FromString(PluginName)));
 		Info.bUseThrobber = false;
 		Info.ExpireDuration = 8.0f;
 		FSlateNotificationManager::Get().AddNotification(Info)->SetCompletionState(SNotificationItem::CS_Success);
@@ -846,71 +740,12 @@ FReply SNewPluginWizard::OnCreatePluginClicked()
 	}
 	else
 	{
-		DeletePluginDirectory(*PluginFolder);
+		const FText Title = LOCTEXT("UnableToCreatePlugin", "Unable to create plugin");
+		FMessageDialog::Open(EAppMsgType::Ok, FailReason, &Title);
 		return FReply::Unhandled();
 	}
 }
 
-bool SNewPluginWizard::CopyFile(const FString& DestinationFile, const FString& SourceFile, TArray<FString>& InOutCreatedFiles)
-{
-	if (IFileManager::Get().Copy(*DestinationFile, *SourceFile, /*bReplaceExisting=*/ false) != COPY_OK)
-	{
-		const FText ErrorMessage = FText::Format(LOCTEXT("ErrorCopyingFile", "Error: Couldn't copy file '{0}' to '{1}'"), FText::AsCultureInvariant(SourceFile), FText::AsCultureInvariant(DestinationFile));
-		PopErrorNotification(ErrorMessage);
-		return false;
-	}
-	else
-	{
-		InOutCreatedFiles.Add(DestinationFile);
-		return true;
-	}
-}
-
-bool SNewPluginWizard::WritePluginDescriptor(const FString& PluginModuleName, const FString& UPluginFilePath, const FWriteDescriptorParams& Params)
-{
-	FPluginDescriptor Descriptor;
-
-	Descriptor.FriendlyName = PluginModuleName;
-	Descriptor.Version = 1;
-	Descriptor.VersionName = TEXT("1.0");
-	Descriptor.Category = TEXT("Other");
-
-	if (DescriptorData.IsValid())
-	{
-		Descriptor.CreatedBy = DescriptorData->CreatedBy;
-		Descriptor.CreatedByURL = DescriptorData->CreatedByURL;
-		Descriptor.Description = DescriptorData->Description;
-		Descriptor.bIsBetaVersion = DescriptorData->bIsBetaVersion;
-	}
-
-	if (Params.bHasModules)
-	{
-		Descriptor.Modules.Add(FModuleDescriptor(*PluginModuleName, Params.ModuleDescriptorType, Params.LoadingPhase));
-	}
-	Descriptor.bCanContainContent = Params.bCanContainContent;
-
-	// Save the descriptor using JSon
-	FText FailReason;
-	if (!Descriptor.Save(UPluginFilePath, FailReason))
-	{
-		PopErrorNotification(FText::Format(LOCTEXT("FailedToWriteDescriptor", "Couldn't save plugin descriptor under {0}"), FText::AsCultureInvariant(UPluginFilePath)));
-		return false;
-	}
-
-	return true;
-}
-
-void SNewPluginWizard::PopErrorNotification(const FText& ErrorMessage)
-{
-	FText Title = LOCTEXT("UnableToCreatePlugin", "Unable to create plugin");
-	FMessageDialog::Open(EAppMsgType::Ok, ErrorMessage, &Title);
-}
-
-void SNewPluginWizard::DeletePluginDirectory(const FString& InPath)
-{
-	IFileManager::Get().DeleteDirectory(*InPath, false, true);
-}
-
 EVisibility SNewPluginWizard::GetShowPluginContentDirectoryVisibility() const
 {
 	EVisibility ShowContentVisibility = EVisibility::Collapsed;
@@ -923,4 +758,4 @@ EVisibility SNewPluginWizard::GetShowPluginContentDirectoryVisibility() const
 	return ShowContentVisibility;
 }
 
-#undef LOCTEXT_NAMESPACE
\ No newline at end of file
+#undef LOCTEXT_NAMESPACE
diff --git a/Engine/Plugins/Editor/PluginBrowser/Source/PluginBrowser/Private/SNewPluginWizard.h b/Engine/Plugins/Editor/PluginBrowser/Source/PluginBrowser/Private/SNewPluginWizard.h
--- a/Engine/Plugins/Editor/PluginBrowser/Source/PluginBrowser/Private/SNewPluginWizard.h
+++ b/Engine/Plugins/Editor/PluginBrowser/Source/PluginBrowser/Private/SNewPluginWizard.h
@@ -24,32 +24,6 @@ DECLARE_LOG_CATEGORY_EXTERN(LogPluginWizard, Log, All);
 
 class SFilePathBlock;
 
-/**
- * Parameters for writing out the descriptor file
- */
-struct FWriteDescriptorParams
-{
-	FWriteDescriptorParams()
-		: bCanContainContent(false)
-		, bHasModules(false)
-		, ModuleDescriptorType(EHostType::Runtime)
-		, LoadingPhase(ELoadingPhase::Default)
-	{
-	}
-
-	/** Can this plugin contain content */
-	bool bCanContainContent;
-
-	/** Does this plugin have Source files? */
-	bool bHasModules;
-
-	/** If this plugin has Source, what is the type of Source included(so it can potentially be excluded in the right builds) */
-	EHostType::Type ModuleDescriptorType;
-
-	/** If this plugin has Source, when should the module be loaded (may need to be earlier than default if used in blueprints) */
-	ELoadingPhase::Type LoadingPhase;
-};
-
 /**
  * A wizard to create a new plugin
  */
@@ -146,36 +120,6 @@ private:
 	 */
 	FReply OnCreatePluginClicked();
 
-	/**
-	 * Copies a file and adds to a list of created files
-	 * @param DestinationFile Where the file will be copied
-	 * @param SourceFile Original file to copy
-	 * @param InOutCreatedFiles Array of created files to add to
-	 * @return Whether the file was copied successfully
-	 */
-	bool CopyFile(const FString& DestinationFile, const FString& SourceFile, TArray<FString>& InOutCreatedFiles);
-
-	/**
-	 * Writes a plugin descriptor file to disk
-	 * @param PluginModuleName		Name of the plugin and its module
-	 * @param UPluginFilePath		Path where the descriptor file should be written
-	 * @param Parmas				Additional parameters for writing out the descriptor file		
-	 * @return Whether the files was written successfully
-	 */
-	bool WritePluginDescriptor(const FString& PluginModuleName, const FString& UPluginFilePath, const FWriteDescriptorParams& Params);
-
-	/**
-	 * Displays an editor pop up error notification
-	 * @param ErrorMessage Message to display in the pop up
-	 */
-	void PopErrorNotification(const FText& ErrorMessage);
-
-	/**
-	 * Helper function to delete given directory
-	 * @param InPath - full path to directory that you want to remove
-	 */
-	void DeletePluginDirectory(const FString& InPath);
-
 	/**
 	 * Generates the expected list view for the plugin wizard, based on the wizard's definition
 	 */
diff --git a/Engine/Source/Developer/AssetTools/AssetTools.Build.cs b/Engine/Source/Developer/AssetTools/AssetTools.Build.cs
--- a/Engine/Source/Developer/AssetTools/AssetTools.Build.cs
+++ b/Engine/Source/Developer/AssetTools/AssetTools.Build.cs
@@ -8,10 +8,12 @@ public class AssetTools : ModuleRules
 	{
 		PrivateIncludePaths.Add("Developer/AssetTools/Private");
 
-		PrivateDependencyModuleNames.AddRange(
+		PublicDependencyModuleNames.AddRange(
 			new string[] {
 				"Core",
 				"CoreUObject",
+				"SlateCore",
+				"UnrealEd",
 			}
 		);
 
@@ -22,11 +24,9 @@ public class AssetTools : ModuleRules
                 "InputCore",
 				"ApplicationCore",
 				"Slate",
-				"SlateCore",
                 "EditorStyle",
 				"SourceControl",
 				"TextureEditor",
-				"UnrealEd",
 				"PropertyEditor",
 				"Kismet",
 				"Landscape",
diff --git a/Engine/Source/Runtime/Projects/Private/PluginDescriptor.cpp b/Engine/Source/Runtime/Projects/Private/PluginDescriptor.cpp
--- a/Engine/Source/Runtime/Projects/Private/PluginDescriptor.cpp
+++ b/Engine/Source/Runtime/Projects/Private/PluginDescriptor.cpp
@@ -171,7 +171,7 @@ bool FPluginDescriptor::Save(const FString& FileName, FText& OutFailReason) cons
 	// Save it to a file
 	if ( !FFileHelper::SaveStringToFile(Text, *FileName) )
 	{
-		OutFailReason = FText::Format( LOCTEXT("FailedToWriteOutputFile", "Failed to write output file '{0}'. Perhaps the file is Read-Only?"), FText::FromString(FileName) );
+		OutFailReason = FText::Format( LOCTEXT("FailedToWriteOutputFile", "Failed to write plugin descriptor file '{0}'. Perhaps the file is Read-Only?"), FText::FromString(FileName) );
 		return false;
 	}
 	return true;
