New plugin wizard improvements
- Now adds created files to source control if it is enabled/active in the project
- Added a new FPluginUtils::CreateAndMountNewPlugin overload that provides more control over the plugin creation

diff --git a/Engine/Plugins/Developer/PluginUtils/Source/PluginUtils/PluginUtils.Build.cs b/Engine/Plugins/Developer/PluginUtils/Source/PluginUtils/PluginUtils.Build.cs
--- a/Engine/Plugins/Developer/PluginUtils/Source/PluginUtils/PluginUtils.Build.cs
+++ b/Engine/Plugins/Developer/PluginUtils/Source/PluginUtils/PluginUtils.Build.cs
@@ -20,6 +20,7 @@ namespace UnrealBuildTool.Rules
 					"AssetRegistry",
 					"AssetTools",
 					"CoreUObject",
+					"SourceControl"
 				}
 			);
 		}
diff --git a/Engine/Plugins/Developer/PluginUtils/Source/PluginUtils/Private/PluginUtils.cpp b/Engine/Plugins/Developer/PluginUtils/Source/PluginUtils/Private/PluginUtils.cpp
--- a/Engine/Plugins/Developer/PluginUtils/Source/PluginUtils/Private/PluginUtils.cpp
+++ b/Engine/Plugins/Developer/PluginUtils/Source/PluginUtils/Private/PluginUtils.cpp
@@ -1,7 +1,7 @@
 // Copyright Epic Games, Inc. All Rights Reserved.
 
 #include "PluginUtils.h"
-
+#include "SourceControlHelpers.h"
 #include "IContentBrowserSingleton.h"
 #include "ContentBrowserModule.h"
 #include "GameProjectUtils.h"
@@ -28,7 +28,7 @@ namespace PluginUtils
 	// The text macro to replace with the actual plugin name when copying files
 	const FString PLUGIN_NAME = TEXT("PLUGIN_NAME");
 
-	bool CopyPluginTemplateFolder(const TCHAR* DestinationDirectory, const TCHAR* Source, const FString& PluginName, FText& FailReason)
+	bool CopyPluginTemplateFolder(const TCHAR* DestinationDirectory, const TCHAR* Source, const FString& PluginName, FText& FailReason, TArray<FString>& InOutFilePathsWritten)
 	{
 		check(DestinationDirectory);
 		check(Source);
@@ -66,13 +66,15 @@ namespace PluginUtils
 			TArray<FString> IgnoredFileTypes;
 			TArray<FString> CopyUnmodifiedFileTypes;
 			FText& FailReason;
+			TArray<FString>& FilePathsWritten;
 
-			FCopyPluginFilesAndDirs(IPlatformFile& InPlatformFile, const TCHAR* InSourceRoot, const TCHAR* InDestRoot, const FString& InPluginName, FText& InFailReason)
+			FCopyPluginFilesAndDirs(IPlatformFile& InPlatformFile, const TCHAR* InSourceRoot, const TCHAR* InDestRoot, const FString& InPluginName, FText& InFailReason, TArray<FString>& InFilePathsWritten)
 				: PlatformFile(InPlatformFile)
 				, SourceRoot(InSourceRoot)
 				, DestRoot(InDestRoot)
 				, PluginName(InPluginName)
 				, FailReason(InFailReason)
+				, FilePathsWritten(InFilePathsWritten)
 			{
 				// Which file types we want to replace instances of PLUGIN_NAME with the new Plugin Name
 				NameReplacementFileTypes.Add(TEXT("cs"));
@@ -164,6 +166,7 @@ namespace PluginUtils
 								return false;
 							}
 						}
+						FilePathsWritten.Add(NewName);
 					}
 				}
 				return true; // continue searching
@@ -171,7 +174,7 @@ namespace PluginUtils
 		};
 
 		// copy plugin files and directories visitor
-		FCopyPluginFilesAndDirs CopyFilesAndDirs(PlatformFile, *SourceDir, *DestDir, PluginName, FailReason);
+		FCopyPluginFilesAndDirs CopyFilesAndDirs(PlatformFile, *SourceDir, *DestDir, PluginName, FailReason, InOutFilePathsWritten);
 
 		// create all files subdirectories and files in subdirectories!
 		return PlatformFile.IterateDirectoryRecursively(*SourceDir, CopyFilesAndDirs);
@@ -366,6 +369,32 @@ FString FPluginUtils::GetPluginResourcesFolder(const FString& PluginLocation, co
 }
 
 TSharedPtr<IPlugin> FPluginUtils::CreateAndMountNewPlugin(const FString& PluginName, const FString& PluginLocation, const FNewPluginParams& CreationParams, const FMountPluginParams& MountParams, FText& FailReason)
+{
+	FNewPluginParamsWithDescriptor ExCreationParams;
+	ExCreationParams.PluginIconPath = CreationParams.PluginIconPath;
+	ExCreationParams.TemplateFolders = CreationParams.TemplateFolders;
+
+	ExCreationParams.Descriptor.FriendlyName = PluginName;
+	ExCreationParams.Descriptor.Version = 1;
+	ExCreationParams.Descriptor.VersionName = TEXT("1.0");
+	ExCreationParams.Descriptor.Category = TEXT("Other");
+	ExCreationParams.Descriptor.CreatedBy = CreationParams.CreatedBy;
+	ExCreationParams.Descriptor.CreatedByURL = CreationParams.CreatedByURL;
+	ExCreationParams.Descriptor.Description = CreationParams.Description;
+	ExCreationParams.Descriptor.bIsBetaVersion = CreationParams.bIsBetaVersion;
+	ExCreationParams.Descriptor.bCanContainContent = CreationParams.bCanContainContent;
+	ExCreationParams.Descriptor.EnabledByDefault = CreationParams.EnabledByDefault;
+	ExCreationParams.Descriptor.bExplicitlyLoaded = CreationParams.bExplicitelyLoaded;
+
+	if (CreationParams.bHasModules)
+	{
+		ExCreationParams.Descriptor.Modules.Add(FModuleDescriptor(*PluginName, CreationParams.ModuleDescriptorType, CreationParams.LoadingPhase));
+	}
+
+	return CreateAndMountNewPlugin(PluginName, PluginLocation, ExCreationParams, MountParams, FailReason);
+}
+
+TSharedPtr<IPlugin> FPluginUtils::CreateAndMountNewPlugin(const FString& PluginName, const FString& PluginLocation, const FNewPluginParamsWithDescriptor& CreationParams, const FMountPluginParams& MountParams, FText& FailReason)
 {
 	// Early validations on new plugin params
 	if (PluginName.IsEmpty())
@@ -380,7 +409,7 @@ TSharedPtr<IPlugin> FPluginUtils::CreateAndMountNewPlugin(const FString& PluginN
 		return nullptr;
 	}
 
-	if (CreationParams.bHasModules && CreationParams.TemplateFolders.Num() == 0)
+	if ((CreationParams.Descriptor.Modules.Num() > 0) && (CreationParams.TemplateFolders.Num() == 0))
 	{
 		FailReason = LOCTEXT("CreateNewPluginParam_NoTemplateFolder", "A template folder must be specified to create a plugin with code");
 		return nullptr;
@@ -398,6 +427,7 @@ TSharedPtr<IPlugin> FPluginUtils::CreateAndMountNewPlugin(const FString& PluginN
 	do
 	{
 		IPlatformFile& PlatformFile = FPlatformFileManager::Get().GetPlatformFile();
+		TArray<FString> NewFilePaths;
 
 		if (!PlatformFile.DirectoryExists(*PluginFolder) && !PlatformFile.CreateDirectoryTree(*PluginFolder))
 		{
@@ -406,7 +436,7 @@ TSharedPtr<IPlugin> FPluginUtils::CreateAndMountNewPlugin(const FString& PluginN
 			break;
 		}
 
-		if (CreationParams.bCanContainContent)
+		if (CreationParams.Descriptor.bCanContainContent)
 		{
 			const FString PluginContentFolder = FPluginUtils::GetPluginContentFolder(PluginLocation, PluginName, /*bFullPath*/ true);
 			if (!PlatformFile.DirectoryExists(*PluginContentFolder) && !PlatformFile.CreateDirectory(*PluginContentFolder))
@@ -417,30 +447,14 @@ TSharedPtr<IPlugin> FPluginUtils::CreateAndMountNewPlugin(const FString& PluginN
 			}
 		}
 
-		FPluginDescriptor Descriptor;
-		Descriptor.FriendlyName = PluginName;
-		Descriptor.Version = 1;
-		Descriptor.VersionName = TEXT("1.0");
-		Descriptor.Category = TEXT("Other");
-		Descriptor.CreatedBy = CreationParams.CreatedBy;
-		Descriptor.CreatedByURL = CreationParams.CreatedByURL;
-		Descriptor.Description = CreationParams.Description;
-		Descriptor.bIsBetaVersion = CreationParams.bIsBetaVersion;
-		Descriptor.bCanContainContent = CreationParams.bCanContainContent;
-		Descriptor.EnabledByDefault = CreationParams.EnabledByDefault;
-		Descriptor.bExplicitlyLoaded = CreationParams.bExplicitelyLoaded;
-		if (CreationParams.bHasModules)
-		{
-			Descriptor.Modules.Add(FModuleDescriptor(*PluginName, CreationParams.ModuleDescriptorType, CreationParams.LoadingPhase));
-		}
-
 		// Write the uplugin file
 		const FString PluginFilePath = FPluginUtils::GetPluginFilePath(PluginLocation, PluginName, /*bFullPath*/ true);
-		if (!Descriptor.Save(PluginFilePath, FailReason))
+		if (!CreationParams.Descriptor.Save(PluginFilePath, FailReason))
 		{
 			bSucceeded = false;
 			break;
 		}
+		NewFilePaths.Add(PluginFilePath);
 
 		// Copy plugin icon
 		if (!CreationParams.PluginIconPath.IsEmpty())
@@ -453,27 +467,32 @@ TSharedPtr<IPlugin> FPluginUtils::CreateAndMountNewPlugin(const FString& PluginN
 				bSucceeded = false;
 				break;
 			}
+			NewFilePaths.Add(DestinationPluginIconPath);
 		}
 
 		// Copy template files
-		GWarn->BeginSlowTask(LOCTEXT("CopyingPluginTemplate", "Copying plugin template files..."), /*ShowProgressDialog*/ true, /*bShowCancelButton*/ false);
-		for (const FString& TemplateFolder : CreationParams.TemplateFolders)
+		if (CreationParams.TemplateFolders.Num() > 0)
 		{
-			if (!PluginUtils::CopyPluginTemplateFolder(*PluginFolder, *TemplateFolder, PluginName, FailReason))
+			GWarn->BeginSlowTask(LOCTEXT("CopyingPluginTemplate", "Copying plugin template files..."), /*ShowProgressDialog*/ true, /*bShowCancelButton*/ false);
+			for (const FString& TemplateFolder : CreationParams.TemplateFolders)
 			{
-				FailReason = FText::Format(LOCTEXT("FailedToCopyPluginTemplate", "Failed to copy plugin template files\nFrom: {0}\nTo: {1}\n{2}"), FText::FromString(FPaths::ConvertRelativePathToFull(TemplateFolder)), FText::FromString(PluginFolder), FailReason);
-				bSucceeded = false;
-				break;
+				if (!PluginUtils::CopyPluginTemplateFolder(*PluginFolder, *TemplateFolder, PluginName, FailReason, NewFilePaths))
+				{
+					FailReason = FText::Format(LOCTEXT("FailedToCopyPluginTemplate", "Failed to copy plugin template files\nFrom: {0}\nTo: {1}\n{2}"), FText::FromString(FPaths::ConvertRelativePathToFull(TemplateFolder)), FText::FromString(PluginFolder), FailReason);
+					bSucceeded = false;
+					break;
+				}
 			}
+			GWarn->EndSlowTask();
 		}
-		GWarn->EndSlowTask();
+
 		if (!bSucceeded)
 		{
 			break;
 		}
 
 		// Compile plugin code
-		if (CreationParams.bHasModules)
+		if (CreationParams.Descriptor.Modules.Num() > 0)
 		{
 			const FString ProjectFileName = IFileManager::Get().ConvertToAbsolutePathForExternalAppForWrite(*FPaths::GetProjectFilePath());
 			const FString Arguments = FString::Printf(TEXT("%s %s %s -Plugin=\"%s\" -Project=\"%s\" -Progress -NoHotReloadFromIDE"), FPlatformMisc::GetUBTTargetName(), FModuleManager::Get().GetUBTConfiguration(), FPlatformMisc::GetUBTPlatform(), *PluginFilePath, *ProjectFileName);
@@ -505,12 +524,20 @@ TSharedPtr<IPlugin> FPluginUtils::CreateAndMountNewPlugin(const FString& PluginN
 		}
 
 		// Fix any content that was added to the plugin
-		if (CreationParams.bCanContainContent)
+		if (CreationParams.Descriptor.bCanContainContent)
 		{	
 			GWarn->BeginSlowTask(LOCTEXT("LoadingContent", "Loading Content..."), /*ShowProgressDialog*/ true, /*bShowCancelButton*/ false);
 			PluginUtils::FixupPluginTemplateAssets(PluginName);
 			GWarn->EndSlowTask();
 		}
+
+		// Add the plugin files to source control if the project is configured for it
+		if (USourceControlHelpers::IsAvailable())
+		{
+			GWarn->BeginSlowTask(LOCTEXT("AddingFilesToSourceControl", "Adding to Source Control..."), /*ShowProgressDialog*/ true, /*bShowCancelButton*/ false);
+			USourceControlHelpers::MarkFilesForAdd(NewFilePaths);
+			GWarn->EndSlowTask();
+		}
 	} while (false);
 
 	if (!bSucceeded)
diff --git a/Engine/Plugins/Developer/PluginUtils/Source/PluginUtils/Public/PluginUtils.h b/Engine/Plugins/Developer/PluginUtils/Source/PluginUtils/Public/PluginUtils.h
--- a/Engine/Plugins/Developer/PluginUtils/Source/PluginUtils/Public/PluginUtils.h
+++ b/Engine/Plugins/Developer/PluginUtils/Source/PluginUtils/Public/PluginUtils.h
@@ -91,6 +91,24 @@ public:
 		ELoadingPhase::Type LoadingPhase = ELoadingPhase::Default;
 	};
 
+	/**
+	 * Parameters for creating a new plugin.
+	 */
+	struct FNewPluginParamsWithDescriptor
+	{
+		/** The description of the plugin */
+		FPluginDescriptor Descriptor;
+
+		/** Path to plugin icon to copy in the plugin resources folder */
+		FString PluginIconPath;
+
+		/**
+		 * Folders containing template files to copy into the plugin folder (Required if Descriptor.Modules is not empty).
+		 * Occurrences of the string PLUGIN_NAME in the filename or file content will be replaced by the plugin name.
+		 */
+		TArray<FString> TemplateFolders;
+	};
+
 	/**
 	 * Parameters for mounting a plugin.
 	 */
@@ -122,6 +140,18 @@ public:
 	 */
 	static TSharedPtr<IPlugin> CreateAndMountNewPlugin(const FString& PluginName, const FString& PluginLocation, const FNewPluginParams& CreationParams, const FMountPluginParams& MountParams, FText& FailReason);
 
+	/**
+	 * Helper to create and mount a new plugin.
+	 * @param PluginName Plugin name
+	 * @param PluginLocation Directory that contains the plugin folder
+	 * @param CreationParams Plugin creation parameters
+	 * @param MountParams Plugin mounting parameters
+	 * @param FailReason Reason the plugin creation/mount failed
+	 * @return The newly created plugin. If something goes wrong during the creation process, the plugin folder gets deleted and null is returned.
+	 * @note Will fail if the plugin already exists
+	 */
+	static TSharedPtr<IPlugin> CreateAndMountNewPlugin(const FString& PluginName, const FString& PluginLocation, const FNewPluginParamsWithDescriptor& CreationParams, const FMountPluginParams& MountParams, FText& FailReason);
+
 	/**
 	 * Load/mount the specified plugin.
 	 * @param PluginName Plugin name
