Improved warning message when a plugin fails to compile. No functional changes.

diff --git a/Engine/Plugins/Developer/PluginUtils/Source/PluginUtils/Private/PluginUtils.cpp b/Engine/Plugins/Developer/PluginUtils/Source/PluginUtils/Private/PluginUtils.cpp
--- a/Engine/Plugins/Developer/PluginUtils/Source/PluginUtils/Private/PluginUtils.cpp
+++ b/Engine/Plugins/Developer/PluginUtils/Source/PluginUtils/Private/PluginUtils.cpp
@@ -498,7 +498,7 @@ TSharedPtr<IPlugin> FPluginUtils::CreateAndMountNewPlugin(const FString& PluginN
 			const FString Arguments = FString::Printf(TEXT("%s %s %s -Plugin=\"%s\" -Project=\"%s\" -Progress -NoHotReloadFromIDE"), FPlatformMisc::GetUBTTargetName(), FModuleManager::Get().GetUBTConfiguration(), FPlatformMisc::GetUBTPlatform(), *PluginFilePath, *ProjectFileName);
 			if (!FDesktopPlatformModule::Get()->RunUnrealBuildTool(FText::Format(LOCTEXT("CompilingPlugin", "Compiling {0} plugin..."), FText::FromString(PluginName)), FPaths::RootDir(), Arguments, GWarn))
 			{
-				FailReason = LOCTEXT("FailedToCompilePlugin", "Failed to compile plugin source code");
+				FailReason = LOCTEXT("FailedToCompilePlugin", "Failed to compile plugin source code. See output log for more information.");
 				bSucceeded = false;
 				break;
 			}
