PluginUtils: add FriendlyName to FNewPluginParams

diff --git a/Engine/Plugins/Developer/PluginUtils/Source/PluginUtils/Private/PluginUtils.cpp b/Engine/Plugins/Developer/PluginUtils/Source/PluginUtils/Private/PluginUtils.cpp
--- a/Engine/Plugins/Developer/PluginUtils/Source/PluginUtils/Private/PluginUtils.cpp
+++ b/Engine/Plugins/Developer/PluginUtils/Source/PluginUtils/Private/PluginUtils.cpp
@@ -380,7 +380,7 @@ TSharedPtr<IPlugin> FPluginUtils::CreateAndMountNewPlugin(const FString& PluginN
 	ExCreationParams.PluginIconPath = CreationParams.PluginIconPath;
 	ExCreationParams.TemplateFolders = CreationParams.TemplateFolders;
 
-	ExCreationParams.Descriptor.FriendlyName = PluginName;
+	ExCreationParams.Descriptor.FriendlyName = CreationParams.FriendlyName.Len() > 0 ? CreationParams.FriendlyName : PluginName;
 	ExCreationParams.Descriptor.Version = 1;
 	ExCreationParams.Descriptor.VersionName = TEXT("1.0");
 	ExCreationParams.Descriptor.Category = TEXT("Other");
@@ -682,7 +682,7 @@ bool FPluginUtils::ValidateNewPluginNameAndLocation(const FString& PluginName, c
 			{
 				if (FailReason)
 				{
-					*FailReason = FText::Format(LOCTEXT("PluginLocationIsFile", "Plugin location is invalid because a file exists at this path\n{0}"), FText::FromString(ExistingFilePath));
+					*FailReason = FText::Format(LOCTEXT("PluginLocationIsFile", "Plugin location is invalid because the following file is in the way\n{0}"), FText::FromString(ExistingFilePath));
 				}
 				return false;
 			}
diff --git a/Engine/Plugins/Developer/PluginUtils/Source/PluginUtils/Public/PluginUtils.h b/Engine/Plugins/Developer/PluginUtils/Source/PluginUtils/Public/PluginUtils.h
--- a/Engine/Plugins/Developer/PluginUtils/Source/PluginUtils/Public/PluginUtils.h
+++ b/Engine/Plugins/Developer/PluginUtils/Source/PluginUtils/Public/PluginUtils.h
@@ -54,6 +54,9 @@ public:
 		/** Hyperlink for the author's website  */
 		FString CreatedByURL;
 
+		/** A friendly name for this plugin. Set to the plugin name by default */
+		FString FriendlyName;
+
 		/** A description for this plugin */
 		FString Description;
 
