Minor improvement on PluginNameAlreadyInUse error message

diff --git a/Engine/Plugins/Developer/PluginUtils/Source/PluginUtils/Private/PluginUtils.cpp b/Engine/Plugins/Developer/PluginUtils/Source/PluginUtils/Private/PluginUtils.cpp
--- a/Engine/Plugins/Developer/PluginUtils/Source/PluginUtils/Private/PluginUtils.cpp
+++ b/Engine/Plugins/Developer/PluginUtils/Source/PluginUtils/Private/PluginUtils.cpp
@@ -694,7 +694,7 @@ bool FPluginUtils::ValidateNewPluginNameAndLocation(const FString& PluginName, c
 	{
 		if (FailReason)
 		{
-			*FailReason = FText::Format(LOCTEXT("PluginNameAlreadyInUse", "Plugin name is already in use\n{0}"), FText::FromString(FPaths::ConvertRelativePathToFull(ExistingPlugin->GetDescriptorFileName())));
+			*FailReason = FText::Format(LOCTEXT("PluginNameAlreadyInUse", "Plugin name is already in use by\n{0}"), FText::FromString(FPaths::ConvertRelativePathToFull(ExistingPlugin->GetDescriptorFileName())));
 		}
 		return false;
 	}
