FPluginUtils::UnloadPlugin(s) utility to unload all packages in a list of plugins and them unmount them.

diff --git a/Engine/Plugins/Developer/PluginUtils/Source/PluginUtils/PluginUtils.Build.cs b/Engine/Plugins/Developer/PluginUtils/Source/PluginUtils/PluginUtils.Build.cs
--- a/Engine/Plugins/Developer/PluginUtils/Source/PluginUtils/PluginUtils.Build.cs
+++ b/Engine/Plugins/Developer/PluginUtils/Source/PluginUtils/PluginUtils.Build.cs
@@ -15,6 +15,7 @@ namespace UnrealBuildTool.Rules
 
 			PrivateDependencyModuleNames.AddRange(
 				new string[] {
+					"UnrealEd",
 					"GameProjectGeneration",
 					"DesktopPlatform",
 					"AssetRegistry",
diff --git a/Engine/Plugins/Developer/PluginUtils/Source/PluginUtils/Private/PluginUtils.cpp b/Engine/Plugins/Developer/PluginUtils/Source/PluginUtils/Private/PluginUtils.cpp
--- a/Engine/Plugins/Developer/PluginUtils/Source/PluginUtils/Private/PluginUtils.cpp
+++ b/Engine/Plugins/Developer/PluginUtils/Source/PluginUtils/Private/PluginUtils.cpp
@@ -1,5 +1,7 @@
 // Copyright Epic Games, Inc. All Rights Reserved.
 
+#if WITH_EDITOR
+
 #include "PluginUtils.h"
 #include "SourceControlHelpers.h"
 #include "IContentBrowserSingleton.h"
@@ -14,6 +16,7 @@
 #include "IAssetTools.h"
 #include "AssetToolsModule.h"
 #include "DesktopPlatformModule.h"
+#include "PackageTools.h"
 #include "HAL/FileManager.h"
 #include "HAL/PlatformFileManager.h"
 #include "Misc/Paths.h"
@@ -21,6 +24,8 @@
 #include "Misc/PackageName.h"
 #include "Misc/FeedbackContext.h"
 
+DEFINE_LOG_CATEGORY_STATIC(LogPluginUtils, Log, All);
+
 #define LOCTEXT_NAMESPACE "PluginUtils"
 
 namespace PluginUtils
@@ -580,6 +585,144 @@ TSharedPtr<IPlugin> FPluginUtils::MountPlugin(const FString& PluginName, const F
 	return PluginUtils::MountPluginInternal(PluginName, PluginLocation, MountParams, FailReason, /*bIsNewPlugin*/ false);
 }
 
+bool FPluginUtils::UnloadPlugin(const TSharedRef<IPlugin>& Plugin, FText* OutFailReason /*= nullptr*/)
+{
+	TArray<TSharedRef<IPlugin>, TInlineAllocator<1>> Plugins;
+	Plugins.Add(Plugin);
+	return UnloadPlugins(Plugins, OutFailReason);
+}
+
+bool FPluginUtils::UnloadPlugin(const FString& PluginName, FText* OutFailReason /*= nullptr*/)
+{
+	if (TSharedPtr<IPlugin> Plugin = IPluginManager::Get().FindPlugin(PluginName))
+	{
+		return UnloadPlugin(Plugin.ToSharedRef(), OutFailReason);
+	}
+	return true;
+}
+
+bool FPluginUtils::UnloadPlugins(const TArrayView<TSharedRef<IPlugin>> Plugins, FText* OutFailReason /*= nullptr*/)
+{
+	if (Plugins.Num() == 0)
+	{
+		return true;
+	}
+
+	TArray<FString> PluginContentMountPoints;
+	PluginContentMountPoints.Reserve(Plugins.Num());
+
+	for (const TSharedRef<IPlugin>& Plugin : Plugins)
+	{
+		if (Plugin->IsEnabled())
+		{
+			FString PluginContentMountPoint = Plugin->GetMountedAssetPath();
+			if (FPackageName::MountPointExists(PluginContentMountPoint))
+			{
+				PluginContentMountPoints.Add(MoveTemp(PluginContentMountPoint));
+			}
+		}
+	}
+
+	// Synchronous scan plugins to make sure we find all their assets.
+	IAssetRegistry& AssetRegistry = FModuleManager::LoadModuleChecked<FAssetRegistryModule>(TEXT("AssetRegistry")).Get();
+
+	if (PluginContentMountPoints.IsEmpty() == false)
+	{
+		AssetRegistry.ScanPathsSynchronous(PluginContentMountPoints, /*bForceRescan=*/ true);
+	}
+
+	// Unload plugin packages
+	{
+		FARFilter ARFilter;
+		ARFilter.PackagePaths.Reserve(PluginContentMountPoints.Num());
+		for (const FString& PluginContentMountPoint : PluginContentMountPoints)
+		{
+			FString PluginRoot = PluginContentMountPoint;
+			PluginRoot.RemoveFromEnd(TEXT("/"));
+			ARFilter.PackagePaths.Add(*PluginRoot);
+		}
+		ARFilter.bRecursivePaths = true;
+
+		TArray<FAssetData> PluginAssets;
+		if (AssetRegistry.GetAssets(ARFilter, PluginAssets))
+		{
+			TSet<UPackage*> PackagesToUnload;
+			PackagesToUnload.Reserve(PluginAssets.Num());
+			for (const FAssetData& AssetData : PluginAssets)
+			{
+				if (UPackage* Package = FindPackage(NULL, *AssetData.PackageName.ToString()))
+				{
+					PackagesToUnload.Add(Package);
+				}
+			}
+
+			if (PackagesToUnload.Num() > 0)
+			{
+				FText FailReason;
+				if (!UPackageTools::UnloadPackages(PackagesToUnload.Array(), FailReason))
+				{
+					// If some packages fail to unload, bail out
+					UE_LOG(LogPluginUtils, Error, TEXT("Plugins cannot be unloaded because some packages cannot be unloaded: %s"), *FailReason.ToString());
+					if (OutFailReason)
+					{
+						*OutFailReason = MoveTemp(FailReason);
+					}
+					return false;
+				}
+			}
+		}
+	}
+
+	// Unmount the plugins
+	//
+	bool bSuccess = true;
+	{
+		FTextBuilder ErrorBuilder;
+		bool bPluginUnmounted = false;
+
+		for (const TSharedRef<IPlugin>& Plugin : Plugins)
+		{
+			if (Plugin->IsEnabled())
+			{
+				bPluginUnmounted = true;
+
+				FText FailReason;
+				if (!IPluginManager::Get().UnmountExplicitlyLoadedPlugin(Plugin->GetName(), &FailReason))
+				{
+					UE_LOG(LogPluginUtils, Error, TEXT("Plugin %s cannot be unloaded: %s"), *Plugin->GetName(), *FailReason.ToString());
+					ErrorBuilder.AppendLine(FailReason);
+					bSuccess = false;
+				}
+			}
+		}
+
+		if (bPluginUnmounted)
+		{
+			IPluginManager::Get().RefreshPluginsList();
+		}
+
+		if (!bSuccess && OutFailReason)
+		{
+			*OutFailReason = ErrorBuilder.ToText();
+		}
+	}
+	return bSuccess;
+}
+
+bool FPluginUtils::UnloadPlugins(const TArrayView<FString> PluginNames, FText* OutFailReason /*= nullptr*/)
+{
+	TArray<TSharedRef<IPlugin>> Plugins;
+	Plugins.Reserve(PluginNames.Num());
+	for (const FString& PluginName : PluginNames)
+	{
+		if (TSharedPtr<IPlugin> Plugin = IPluginManager::Get().FindPlugin(PluginName))
+		{
+			Plugins.Add(Plugin.ToSharedRef());
+		}
+	}
+	return UnloadPlugins(Plugins, OutFailReason);
+}
+
 bool FPluginUtils::AddToPluginSearchPathIfNeeded(const FString& Dir, bool bRefreshPlugins, bool bUpdateProjectFile)
 {
 	bool bSearchPathChanged = false;
@@ -759,3 +902,5 @@ bool FPluginUtils::IsValidPluginName(const FString& PluginName, FText* FailReaso
 }
 
 #undef LOCTEXT_NAMESPACE
+
+#endif //if WITH_EDITOR
\ No newline at end of file
diff --git a/Engine/Plugins/Developer/PluginUtils/Source/PluginUtils/Public/PluginUtils.h b/Engine/Plugins/Developer/PluginUtils/Source/PluginUtils/Public/PluginUtils.h
--- a/Engine/Plugins/Developer/PluginUtils/Source/PluginUtils/Public/PluginUtils.h
+++ b/Engine/Plugins/Developer/PluginUtils/Source/PluginUtils/Public/PluginUtils.h
@@ -2,6 +2,8 @@
 
 #pragma once
 
+#if WITH_EDITOR
+
 #include "CoreMinimal.h"
 #include "PluginDescriptor.h"
 #include "ModuleDescriptor.h"
@@ -166,6 +168,38 @@ public:
 	 */
 	static TSharedPtr<IPlugin> MountPlugin(const FString& PluginName, const FString& PluginLocation, const FMountPluginParams& MountParams, FText& FailReason);
 
+	/**
+	 * Unload assets from the specified plugin and unmount it
+	 * @param Plugin Plugin to unload
+	 * @param OutFailReason Outputs the reason of the failure if any
+	 * @return true on success
+	 */
+	static bool UnloadPlugin(const TSharedRef<IPlugin>& Plugin, FText* OutFailReason = nullptr);
+
+	/**
+	 * Unload assets from the specified plugin and unmount it
+	 * @param PluginName Name of the plugin to unload
+	 * @param OutFailReason Outputs the reason of the failure if any
+	 * @return true on success
+	 */
+	static bool UnloadPlugin(const FString& PluginName, FText* OutFailReason = nullptr);
+
+	/**
+	 * Unload assets from the specified plugins and unmount them
+	 * @param Plugins Plugins to unload
+	 * @param OutFailReason Outputs the reason of the failure if any
+	 * @return true on success
+	 */
+	static bool UnloadPlugins(const TArrayView<TSharedRef<IPlugin>> Plugins, FText* OutFailReason = nullptr);
+
+	/**
+	 * Unload assets from the specified plugins and unmount them
+	 * @param PluginNames Names of the plugins to unload
+	 * @param OutFailReason Outputs the reason of the failure if any
+	 * @return true on success
+	 */
+	static bool UnloadPlugins(const TArrayView<FString> PluginNames, FText* OutFailReason = nullptr);
+
 	/**
 	 * Adds a directory to the list of paths that are recursively searched for plugins, 
 	 * if that directory isn't already under the search paths.
@@ -194,3 +228,5 @@ public:
 	 */
 	static bool IsValidPluginName(const FString& PluginName, FText* FailReason = nullptr, const FText* PluginTermReplacement = nullptr);
 };
+
+#endif //if WITH_EDITOR
