FPluginUtils::UnloadPlugins TConstArrayView param

diff --git a/Engine/Plugins/Developer/PluginUtils/Source/PluginUtils/Private/PluginUtils.cpp b/Engine/Plugins/Developer/PluginUtils/Source/PluginUtils/Private/PluginUtils.cpp
--- a/Engine/Plugins/Developer/PluginUtils/Source/PluginUtils/Private/PluginUtils.cpp
+++ b/Engine/Plugins/Developer/PluginUtils/Source/PluginUtils/Private/PluginUtils.cpp
@@ -587,9 +587,7 @@ TSharedPtr<IPlugin> FPluginUtils::MountPlugin(const FString& PluginName, const F
 
 bool FPluginUtils::UnloadPlugin(const TSharedRef<IPlugin>& Plugin, FText* OutFailReason /*= nullptr*/)
 {
-	TArray<TSharedRef<IPlugin>, TInlineAllocator<1>> Plugins;
-	Plugins.Add(Plugin);
-	return UnloadPlugins(Plugins, OutFailReason);
+	return UnloadPlugins({ Plugin }, OutFailReason);
 }
 
 bool FPluginUtils::UnloadPlugin(const FString& PluginName, FText* OutFailReason /*= nullptr*/)
@@ -601,7 +599,7 @@ bool FPluginUtils::UnloadPlugin(const FString& PluginName, FText* OutFailReason
 	return true;
 }
 
-bool FPluginUtils::UnloadPlugins(const TArrayView<TSharedRef<IPlugin>> Plugins, FText* OutFailReason /*= nullptr*/)
+bool FPluginUtils::UnloadPlugins(const TConstArrayView<TSharedRef<IPlugin>> Plugins, FText* OutFailReason /*= nullptr*/)
 {
 	if (Plugins.Num() == 0)
 	{
@@ -709,7 +707,7 @@ bool FPluginUtils::UnloadPlugins(const TArrayView<TSharedRef<IPlugin>> Plugins,
 	return bSuccess;
 }
 
-bool FPluginUtils::UnloadPlugins(const TArrayView<FString> PluginNames, FText* OutFailReason /*= nullptr*/)
+bool FPluginUtils::UnloadPlugins(const TConstArrayView<FString> PluginNames, FText* OutFailReason /*= nullptr*/)
 {
 	TArray<TSharedRef<IPlugin>> Plugins;
 	Plugins.Reserve(PluginNames.Num());
diff --git a/Engine/Plugins/Developer/PluginUtils/Source/PluginUtils/Public/PluginUtils.h b/Engine/Plugins/Developer/PluginUtils/Source/PluginUtils/Public/PluginUtils.h
--- a/Engine/Plugins/Developer/PluginUtils/Source/PluginUtils/Public/PluginUtils.h
+++ b/Engine/Plugins/Developer/PluginUtils/Source/PluginUtils/Public/PluginUtils.h
@@ -190,7 +190,7 @@ public:
 	 * @param OutFailReason Outputs the reason of the failure if any
 	 * @return true on success
 	 */
-	static bool UnloadPlugins(const TArrayView<TSharedRef<IPlugin>> Plugins, FText* OutFailReason = nullptr);
+	static bool UnloadPlugins(const TConstArrayView<TSharedRef<IPlugin>> Plugins, FText* OutFailReason = nullptr);
 
 	/**
 	 * Unload assets from the specified plugins and unmount them
@@ -198,7 +198,7 @@ public:
 	 * @param OutFailReason Outputs the reason of the failure if any
 	 * @return true on success
 	 */
-	static bool UnloadPlugins(const TArrayView<FString> PluginNames, FText* OutFailReason = nullptr);
+	static bool UnloadPlugins(const TConstArrayView<FString> PluginNames, FText* OutFailReason = nullptr);
 
 	/**
 	 * Adds a directory to the list of paths that are recursively searched for plugins, 
