Add PluginUtils::FindLoadedPlugin
PluginUtils::LoadPlugin improvements and API cleanup:
- Add SynchronousAssetsScan and OutAlreadyLoaded params
- Make all loading options false by default (caller opts-in on whatever it wants instead of opting out on some)
- Put OutFailReason in the loading param struct
- Deprecate MountPlugin API and rename it LoadPlugin to mirror UnloadPlugin terminology

diff --git a/Engine/Plugins/Developer/PluginUtils/Source/PluginUtils/Private/PluginUtils.cpp b/Engine/Plugins/Developer/PluginUtils/Source/PluginUtils/Private/PluginUtils.cpp
--- a/Engine/Plugins/Developer/PluginUtils/Source/PluginUtils/Private/PluginUtils.cpp
+++ b/Engine/Plugins/Developer/PluginUtils/Source/PluginUtils/Private/PluginUtils.cpp
@@ -33,7 +33,19 @@ namespace PluginUtils
 	// The text macro to replace with the actual plugin name when copying files
 	const FString PLUGIN_NAME = TEXT("PLUGIN_NAME");
 
-	bool CopyPluginTemplateFolder(const TCHAR* DestinationDirectory, const TCHAR* Source, const FString& PluginName, FText& FailReason, TArray<FString>& InOutFilePathsWritten)
+	PRAGMA_DISABLE_DEPRECATION_WARNINGS
+	FPluginUtils::FLoadPluginParams ConvertToLoadPluginParams(const FPluginUtils::FMountPluginParams& MountParams, FText* FailReason)
+	{
+		FPluginUtils::FLoadPluginParams LoadParams;
+		LoadParams.bSelectInContentBrowser = MountParams.bSelectInContentBrowser;
+		LoadParams.bEnablePluginInProject = MountParams.bEnablePluginInProject;
+		LoadParams.bUpdateProjectPluginSearchPath = MountParams.bUpdateProjectPluginSearchPath;
+		LoadParams.OutFailReason = FailReason;
+		return LoadParams;
+	}
+	PRAGMA_ENABLE_DEPRECATION_WARNINGS
+
+	bool CopyPluginTemplateFolder(const TCHAR* DestinationDirectory, const TCHAR* Source, const FString& PluginName, TArray<FString>& InOutFilePathsWritten, FText* OutFailReason = nullptr)
 	{
 		check(DestinationDirectory);
 		check(Source);
@@ -49,14 +61,20 @@ namespace PluginUtils
 		// Does Source dir exist?
 		if (!PlatformFile.DirectoryExists(*SourceDir))
 		{
-			FailReason = FText::Format(LOCTEXT("InvalidPluginTemplateFolder", "Plugin template folder doesn't exist\n{0}"), FText::FromString(FPaths::ConvertRelativePathToFull(SourceDir)));
+			if (OutFailReason)
+			{
+				*OutFailReason = FText::Format(LOCTEXT("InvalidPluginTemplateFolder", "Plugin template folder doesn't exist\n{0}"), FText::FromString(FPaths::ConvertRelativePathToFull(SourceDir)));
+			}
 			return false;
 		}
 
 		// Destination directory exists already or can be created ?
 		if (!PlatformFile.DirectoryExists(*DestDir) && !PlatformFile.CreateDirectoryTree(*DestDir))
 		{
-			FailReason = FText::Format(LOCTEXT("FailedToCreateDestinationFolder", "Failed to create destination folder\n{0}"), FText::FromString(FPaths::ConvertRelativePathToFull(DestDir)));
+			if (OutFailReason)
+			{
+				*OutFailReason = FText::Format(LOCTEXT("FailedToCreateDestinationFolder", "Failed to create destination folder\n{0}"), FText::FromString(FPaths::ConvertRelativePathToFull(DestDir)));
+			}
 			return false;
 		}
 
@@ -70,16 +88,16 @@ namespace PluginUtils
 			TArray<FString> NameReplacementFileTypes;
 			TArray<FString> IgnoredFileTypes;
 			TArray<FString> CopyUnmodifiedFileTypes;
-			FText& FailReason;
 			TArray<FString>& FilePathsWritten;
+			FText* FailReason;
 
-			FCopyPluginFilesAndDirs(IPlatformFile& InPlatformFile, const TCHAR* InSourceRoot, const TCHAR* InDestRoot, const FString& InPluginName, FText& InFailReason, TArray<FString>& InFilePathsWritten)
+			FCopyPluginFilesAndDirs(IPlatformFile& InPlatformFile, const TCHAR* InSourceRoot, const TCHAR* InDestRoot, const FString& InPluginName, TArray<FString>& InFilePathsWritten, FText* InFailReason)
 				: PlatformFile(InPlatformFile)
 				, SourceRoot(InSourceRoot)
 				, DestRoot(InDestRoot)
 				, PluginName(InPluginName)
-				, FailReason(InFailReason)
 				, FilePathsWritten(InFilePathsWritten)
+				, FailReason(InFailReason)
 			{
 				// Which file types we want to replace instances of PLUGIN_NAME with the new Plugin Name
 				NameReplacementFileTypes.Add(TEXT("cs"));
@@ -111,7 +129,10 @@ namespace PluginUtils
 					// create new directory structure
 					if (!PlatformFile.CreateDirectoryTree(*NewName) && !PlatformFile.DirectoryExists(*NewName))
 					{
-						FailReason = FText::Format(LOCTEXT("FailedToCreatePluginSubFolder", "Failed to create plugin subfolder\n{0}"), FText::FromString(FPaths::ConvertRelativePathToFull(NewName)));
+						if (FailReason)
+						{
+							*FailReason = FText::Format(LOCTEXT("FailedToCreatePluginSubFolder", "Failed to create plugin subfolder\n{0}"), FText::FromString(FPaths::ConvertRelativePathToFull(NewName)));
+						}
 						return false;
 					}
 				}
@@ -142,7 +163,10 @@ namespace PluginUtils
 							FString OutFileContents;
 							if (!FFileHelper::LoadFileToString(OutFileContents, FilenameOrDirectory))
 							{
-								FailReason = FText::Format(LOCTEXT("FailedToReadPluginTemplateFile", "Failed to read plugin template file\n{0}"), FText::FromString(FPaths::ConvertRelativePathToFull(FilenameOrDirectory)));
+								if (FailReason)
+								{
+									*FailReason = FText::Format(LOCTEXT("FailedToReadPluginTemplateFile", "Failed to read plugin template file\n{0}"), FText::FromString(FPaths::ConvertRelativePathToFull(FilenameOrDirectory)));
+								}
 								return false;
 							}
 
@@ -157,7 +181,10 @@ namespace PluginUtils
 
 							if (!FFileHelper::SaveStringToFile(OutFileContents, *NewName))
 							{
-								FailReason = FText::Format(LOCTEXT("FailedToWritePluginFile", "Failed to write plugin file\n{0}"), FText::FromString(FPaths::ConvertRelativePathToFull(NewName)));
+								if (FailReason)
+								{
+									*FailReason = FText::Format(LOCTEXT("FailedToWritePluginFile", "Failed to write plugin file\n{0}"), FText::FromString(FPaths::ConvertRelativePathToFull(NewName)));
+								}
 								return false;
 							}
 						}
@@ -167,7 +194,10 @@ namespace PluginUtils
 							if (!PlatformFile.CopyFile(*NewName, FilenameOrDirectory))
 							{
 								// Not all files could be copied
-								FailReason = FText::Format(LOCTEXT("FailedToCopyPluginTemplateFile", "Failed to copy plugin template file\nFrom: {0}\nTo: {1}"), FText::FromString(FPaths::ConvertRelativePathToFull(FilenameOrDirectory)), FText::FromString(FPaths::ConvertRelativePathToFull(NewName)));
+								if (FailReason)
+								{
+									*FailReason = FText::Format(LOCTEXT("FailedToCopyPluginTemplateFile", "Failed to copy plugin template file\nFrom: {0}\nTo: {1}"), FText::FromString(FPaths::ConvertRelativePathToFull(FilenameOrDirectory)), FText::FromString(FPaths::ConvertRelativePathToFull(NewName)));
+								}
 								return false;
 							}
 						}
@@ -179,7 +209,7 @@ namespace PluginUtils
 		};
 
 		// copy plugin files and directories visitor
-		FCopyPluginFilesAndDirs CopyFilesAndDirs(PlatformFile, *SourceDir, *DestDir, PluginName, FailReason, InOutFilePathsWritten);
+		FCopyPluginFilesAndDirs CopyFilesAndDirs(PlatformFile, *SourceDir, *DestDir, PluginName, InOutFilePathsWritten, OutFailReason);
 
 		// create all files subdirectories and files in subdirectories!
 		return PlatformFile.IterateDirectoryRecursively(*SourceDir, CopyFilesAndDirs);
@@ -232,7 +262,7 @@ namespace PluginUtils
 
 				if (FilesToScan.Num() > 0)
 				{
-					IAssetRegistry& AssetRegistry = FModuleManager::LoadModuleChecked<FAssetRegistryModule>("AssetRegistry").Get();
+					IAssetRegistry& AssetRegistry = FModuleManager::LoadModuleChecked<FAssetRegistryModule>(AssetRegistryConstants::ModuleName).Get();
 					AssetRegistry.ScanFilesSynchronous(FilesToScan);
 
 					for (const FString& File : FilesToScan)
@@ -274,19 +304,17 @@ namespace PluginUtils
 		}
 	}
 
-	TSharedPtr<IPlugin> MountPluginInternal(const FString& PluginName, const FString& PluginLocation, const FPluginUtils::FMountPluginParams& MountParams, FText& FailReason, const bool bIsNewPlugin)
+	TSharedPtr<IPlugin> LoadPluginInternal(const FString& PluginName, const FString& PluginLocation, const FString& PluginFilePath, FPluginUtils::FLoadPluginParams& LoadParams, const bool bIsNewPlugin)
 	{
-		ensure(!PluginLocation.IsEmpty());
+		LoadParams.bOutAlreadyLoaded = false;
 
-		const FString PluginFilePath = FPluginUtils::GetPluginFilePath(PluginLocation, PluginName, /*bFullPath*/ true);
-
-		if (MountParams.bUpdateProjectPluginSearchPath)
+		if (LoadParams.bUpdateProjectPluginSearchPath)
 		{
 			FPluginUtils::AddToPluginSearchPathIfNeeded(PluginLocation, /*bRefreshPlugins=*/true, /*bUpdateProjectFile=*/true);
 		}
 		else
 		{
-			if (!IPluginManager::Get().AddToPluginsList(PluginFilePath, &FailReason))
+			if (!IPluginManager::Get().AddToPluginsList(PluginFilePath, LoadParams.OutFailReason))
 			{
 				return nullptr;
 			}
@@ -296,48 +324,83 @@ namespace PluginUtils
 		TSharedPtr<IPlugin> Plugin = IPluginManager::Get().FindPlugin(PluginName);
 		if (!Plugin)
 		{
-			FailReason = FText::Format(LOCTEXT("FailedToRegisterPlugin", "Failed to register plugin\n{0}"), FText::FromString(FPluginUtils::GetPluginFilePath(PluginLocation, PluginName, /*bFullPath*/ true)));
+			if (LoadParams.OutFailReason)
+			{
+				*LoadParams.OutFailReason = FText::Format(LOCTEXT("FailedToRegisterPlugin", "Failed to register plugin\n{0}"), FText::FromString(PluginFilePath));
+			}
 			return nullptr;
 		}
 
 		// Double check the path matches
 		if (!FPaths::IsSamePath(Plugin->GetDescriptorFileName(), PluginFilePath))
 		{
-			const FString PluginFilePathFull = FPaths::ConvertRelativePathToFull(Plugin->GetDescriptorFileName());
-			FailReason = FText::Format(LOCTEXT("PluginNameAlreadyUsed", "There's already a plugin named {0} at this location:\n{1}"), FText::FromString(PluginName), FText::FromString(PluginFilePathFull));
+			if (LoadParams.OutFailReason)
+			{
+				const FString PluginFilePathFull = FPaths::ConvertRelativePathToFull(Plugin->GetDescriptorFileName());
+				*LoadParams.OutFailReason = FText::Format(LOCTEXT("PluginNameAlreadyUsed", "There's already a plugin named {0} at this location:\n{1}"), FText::FromString(PluginName), FText::FromString(PluginFilePathFull));
+			}
 			return nullptr;
 		}
 
+		const FString PluginRootFolder = Plugin->CanContainContent() ? Plugin->GetMountedAssetPath() : FString();
+
+		LoadParams.bOutAlreadyLoaded = !bIsNewPlugin && Plugin->IsEnabled() && (PluginRootFolder.IsEmpty() || FPackageName::MountPointExists(PluginRootFolder));
+
 		// Enable this plugin in the project
-		if (MountParams.bEnablePluginInProject && !IProjectManager::Get().SetPluginEnabled(PluginName, true, FailReason))
+		if (LoadParams.bEnablePluginInProject)
 		{
-			FailReason = FText::Format(LOCTEXT("FailedToEnablePluginInProject", "Failed to enable plugin in current project\n{0}"), FailReason);
-			return nullptr;
+			FText FailReason;
+			if (!IProjectManager::Get().SetPluginEnabled(PluginName, true, FailReason))
+			{
+				if (LoadParams.OutFailReason)
+				{
+					*LoadParams.OutFailReason = FText::Format(LOCTEXT("FailedToEnablePluginInProject", "Failed to enable plugin in current project\n{0}"), FailReason);
+				}
+				return nullptr;
+			}
 		}
 
-		// Mount the new plugin (mount content folder if any and load modules if any)
-		if (bIsNewPlugin)
-		{
-			IPluginManager::Get().MountNewlyCreatedPlugin(PluginName);
-		}
-		else
+		if (!LoadParams.bOutAlreadyLoaded)
 		{
-			IPluginManager::Get().MountExplicitlyLoadedPlugin(PluginName);
-		}
+			// Mount the new plugin (mount content folder if any and load modules if any)
+			if (bIsNewPlugin)
+			{
+				IPluginManager::Get().MountNewlyCreatedPlugin(PluginName);
+			}
+			else
+			{
+				IPluginManager::Get().MountExplicitlyLoadedPlugin(PluginName);
+			}
 
-		if (!Plugin->IsEnabled())
-		{
-			FailReason = FText::Format(LOCTEXT("FailedToEnablePlugin", "Failed to enable plugin because it is not configured as bExplicitlyLoaded=true\n{0}"), FText::FromString(FPluginUtils::GetPluginFilePath(PluginLocation, PluginName, /*bFullPath*/ true)));
-			return nullptr;
+			if (!Plugin->IsEnabled())
+			{
+				if (LoadParams.OutFailReason)
+				{
+					*LoadParams.OutFailReason = FText::Format(LOCTEXT("FailedToEnablePlugin", "Failed to enable plugin because it is not configured as bExplicitlyLoaded=true\n{0}"), FText::FromString(PluginFilePath));
+				}
+				return nullptr;
+			}
 		}
 
-		// Select plugin Content folder in content browser
-		if (MountParams.bSelectInContentBrowser && Plugin->CanContainContent() && !IsRunningCommandlet())
+		const bool bSelectInContentBrowser = LoadParams.bSelectInContentBrowser && !IsRunningCommandlet();
+
+		if ((bSelectInContentBrowser || LoadParams.bSynchronousAssetsScan) && !PluginRootFolder.IsEmpty() && (LoadParams.bOutAlreadyLoaded || FPackageName::MountPointExists(PluginRootFolder)))
 		{
-			IContentBrowserSingleton& ContentBrowser = FModuleManager::LoadModuleChecked<FContentBrowserModule>("ContentBrowser").Get();
-			const bool bIsEnginePlugin = FPaths::IsUnderDirectory(PluginLocation, FPaths::EnginePluginsDir());
-			ContentBrowser.ForceShowPluginContent(bIsEnginePlugin);
-			ContentBrowser.SetSelectedPaths({ Plugin->GetMountedAssetPath() }, /*bNeedsRefresh*/ true);
+			if (LoadParams.bSynchronousAssetsScan)
+			{
+				// Scan plugin assets
+				IAssetRegistry& AssetRegistry = FModuleManager::LoadModuleChecked<FAssetRegistryModule>(AssetRegistryConstants::ModuleName).Get();
+				AssetRegistry.ScanPathsSynchronous({ PluginRootFolder }, /*bForceRescan=*/ true);
+			}
+
+			if (bSelectInContentBrowser)
+			{
+				// Select plugin root folder in content browser
+				IContentBrowserSingleton& ContentBrowser = FModuleManager::LoadModuleChecked<FContentBrowserModule>("ContentBrowser").Get();
+				const bool bIsEnginePlugin = FPaths::IsUnderDirectory(PluginLocation, FPaths::EnginePluginsDir());
+				ContentBrowser.ForceShowPluginContent(bIsEnginePlugin);
+				ContentBrowser.SetSelectedPaths({ PluginRootFolder }, /*bNeedsRefresh*/ true);
+			}
 		}
 
 		return Plugin;
@@ -388,7 +451,7 @@ FString FPluginUtils::GetPluginResourcesFolder(const FString& PluginLocation, co
 	return PluginResourcesFolder;
 }
 
-TSharedPtr<IPlugin> FPluginUtils::CreateAndMountNewPlugin(const FString& PluginName, const FString& PluginLocation, const FNewPluginParams& CreationParams, const FMountPluginParams& MountParams, FText& FailReason)
+TSharedPtr<IPlugin> FPluginUtils::CreateAndLoadNewPlugin(const FString& PluginName, const FString& PluginLocation, const FNewPluginParams& CreationParams, FLoadPluginParams& LoadParams)
 {
 	FNewPluginParamsWithDescriptor ExCreationParams;
 	ExCreationParams.PluginIconPath = CreationParams.PluginIconPath;
@@ -411,31 +474,40 @@ TSharedPtr<IPlugin> FPluginUtils::CreateAndMountNewPlugin(const FString& PluginN
 		ExCreationParams.Descriptor.Modules.Add(FModuleDescriptor(*PluginName, CreationParams.ModuleDescriptorType, CreationParams.LoadingPhase));
 	}
 
-	return CreateAndMountNewPlugin(PluginName, PluginLocation, ExCreationParams, MountParams, FailReason);
+	return CreateAndLoadNewPlugin(PluginName, PluginLocation, ExCreationParams, LoadParams);
 }
 
-TSharedPtr<IPlugin> FPluginUtils::CreateAndMountNewPlugin(const FString& PluginName, const FString& PluginLocation, const FNewPluginParamsWithDescriptor& CreationParams, const FMountPluginParams& MountParams, FText& FailReason)
+TSharedPtr<IPlugin> FPluginUtils::CreateAndLoadNewPlugin(const FString& PluginName, const FString& PluginLocation, const FNewPluginParamsWithDescriptor& CreationParams, FLoadPluginParams& LoadParams)
 {
 	// Early validations on new plugin params
 	if (PluginName.IsEmpty())
 	{
-		FailReason = LOCTEXT("CreateNewPluginParam_NoPluginName", "Missing plugin name");
+		if (LoadParams.OutFailReason)
+		{
+			*LoadParams.OutFailReason = LOCTEXT("CreateNewPluginParam_NoPluginName", "Missing plugin name");
+		}
 		return nullptr;
 	}
 
 	if (PluginLocation.IsEmpty())
 	{
-		FailReason = LOCTEXT("CreateNewPluginParam_NoPluginLocation", "Missing plugin location");
+		if (LoadParams.OutFailReason)
+		{
+			*LoadParams.OutFailReason = LOCTEXT("CreateNewPluginParam_NoPluginLocation", "Missing plugin location");
+		}
 		return nullptr;
 	}
 
 	if ((CreationParams.Descriptor.Modules.Num() > 0) && (CreationParams.TemplateFolders.Num() == 0))
 	{
-		FailReason = LOCTEXT("CreateNewPluginParam_NoTemplateFolder", "A template folder must be specified to create a plugin with code");
+		if (LoadParams.OutFailReason)
+		{
+			*LoadParams.OutFailReason = LOCTEXT("CreateNewPluginParam_NoTemplateFolder", "A template folder must be specified to create a plugin with code");
+		}
 		return nullptr;
 	}
 
-	if (!FPluginUtils::ValidateNewPluginNameAndLocation(PluginName, PluginLocation, &FailReason))
+	if (!FPluginUtils::ValidateNewPluginNameAndLocation(PluginName, PluginLocation, LoadParams.OutFailReason))
 	{
 		return nullptr;
 	}
@@ -451,7 +523,10 @@ TSharedPtr<IPlugin> FPluginUtils::CreateAndMountNewPlugin(const FString& PluginN
 
 		if (!PlatformFile.DirectoryExists(*PluginFolder) && !PlatformFile.CreateDirectoryTree(*PluginFolder))
 		{
-			FailReason = FText::Format(LOCTEXT("FailedToCreatePluginFolder", "Failed to create plugin folder\n{0}"), FText::FromString(PluginFolder));
+			if (LoadParams.OutFailReason)
+			{
+				*LoadParams.OutFailReason = FText::Format(LOCTEXT("FailedToCreatePluginFolder", "Failed to create plugin folder\n{0}"), FText::FromString(PluginFolder));
+			}
 			bSucceeded = false;
 			break;
 		}
@@ -461,7 +536,10 @@ TSharedPtr<IPlugin> FPluginUtils::CreateAndMountNewPlugin(const FString& PluginN
 			const FString PluginContentFolder = FPluginUtils::GetPluginContentFolder(PluginLocation, PluginName, /*bFullPath*/ true);
 			if (!PlatformFile.DirectoryExists(*PluginContentFolder) && !PlatformFile.CreateDirectory(*PluginContentFolder))
 			{
-				FailReason = FText::Format(LOCTEXT("FailedToCreatePluginContentFolder", "Failed to create plugin Content folder\n{0}"), FText::FromString(PluginContentFolder));
+				if (LoadParams.OutFailReason)
+				{
+					*LoadParams.OutFailReason = FText::Format(LOCTEXT("FailedToCreatePluginContentFolder", "Failed to create plugin Content folder\n{0}"), FText::FromString(PluginContentFolder));
+				}
 				bSucceeded = false;
 				break;
 			}
@@ -469,7 +547,7 @@ TSharedPtr<IPlugin> FPluginUtils::CreateAndMountNewPlugin(const FString& PluginN
 
 		// Write the uplugin file
 		const FString PluginFilePath = FPluginUtils::GetPluginFilePath(PluginLocation, PluginName, /*bFullPath*/ true);
-		if (!CreationParams.Descriptor.Save(PluginFilePath, FailReason))
+		if (!CreationParams.Descriptor.Save(PluginFilePath, LoadParams.OutFailReason))
 		{
 			bSucceeded = false;
 			break;
@@ -483,7 +561,10 @@ TSharedPtr<IPlugin> FPluginUtils::CreateAndMountNewPlugin(const FString& PluginN
 			const FString DestinationPluginIconPath = FPaths::Combine(ResourcesFolder, TEXT("Icon128.png"));
 			if (IFileManager::Get().Copy(*DestinationPluginIconPath, *CreationParams.PluginIconPath, /*bReplaceExisting=*/ false) != COPY_OK)
 			{
-				FailReason = FText::Format(LOCTEXT("FailedToCopyPluginIcon", "Failed to copy plugin icon\nFrom: {0}\nTo: {1}"), FText::FromString(FPaths::ConvertRelativePathToFull(CreationParams.PluginIconPath)), FText::FromString(DestinationPluginIconPath));
+				if (LoadParams.OutFailReason)
+				{
+					*LoadParams.OutFailReason = FText::Format(LOCTEXT("FailedToCopyPluginIcon", "Failed to copy plugin icon\nFrom: {0}\nTo: {1}"), FText::FromString(FPaths::ConvertRelativePathToFull(CreationParams.PluginIconPath)), FText::FromString(DestinationPluginIconPath));
+				}
 				bSucceeded = false;
 				break;
 			}
@@ -496,9 +577,12 @@ TSharedPtr<IPlugin> FPluginUtils::CreateAndMountNewPlugin(const FString& PluginN
 			GWarn->BeginSlowTask(LOCTEXT("CopyingPluginTemplate", "Copying plugin template files..."), /*ShowProgressDialog*/ true, /*bShowCancelButton*/ false);
 			for (const FString& TemplateFolder : CreationParams.TemplateFolders)
 			{
-				if (!PluginUtils::CopyPluginTemplateFolder(*PluginFolder, *TemplateFolder, PluginName, FailReason, NewFilePaths))
+				if (!PluginUtils::CopyPluginTemplateFolder(*PluginFolder, *TemplateFolder, PluginName, NewFilePaths, LoadParams.OutFailReason))
 				{
-					FailReason = FText::Format(LOCTEXT("FailedToCopyPluginTemplate", "Failed to copy plugin template files\nFrom: {0}\nTo: {1}\n{2}"), FText::FromString(FPaths::ConvertRelativePathToFull(TemplateFolder)), FText::FromString(PluginFolder), FailReason);
+					if (LoadParams.OutFailReason)
+					{
+						*LoadParams.OutFailReason = FText::Format(LOCTEXT("FailedToCopyPluginTemplate", "Failed to copy plugin template files\nFrom: {0}\nTo: {1}\n{2}"), FText::FromString(FPaths::ConvertRelativePathToFull(TemplateFolder)), FText::FromString(PluginFolder), *LoadParams.OutFailReason);
+					}
 					bSucceeded = false;
 					break;
 				}
@@ -518,7 +602,10 @@ TSharedPtr<IPlugin> FPluginUtils::CreateAndMountNewPlugin(const FString& PluginN
 			const FString Arguments = FString::Printf(TEXT("%s %s %s -Plugin=\"%s\" -Project=\"%s\" -Progress -NoHotReloadFromIDE"), FPlatformMisc::GetUBTTargetName(), FModuleManager::Get().GetUBTConfiguration(), FPlatformMisc::GetUBTPlatform(), *PluginFilePath, *ProjectFileName);
 			if (!FDesktopPlatformModule::Get()->RunUnrealBuildTool(FText::Format(LOCTEXT("CompilingPlugin", "Compiling {0} plugin..."), FText::FromString(PluginName)), FPaths::RootDir(), Arguments, GWarn))
 			{
-				FailReason = LOCTEXT("FailedToCompilePlugin", "Failed to compile plugin source code. See output log for more information.");
+				if (LoadParams.OutFailReason)
+				{
+					*LoadParams.OutFailReason = LOCTEXT("FailedToCompilePlugin", "Failed to compile plugin source code. See output log for more information.");
+				}
 				bSucceeded = false;
 				break;
 			}
@@ -529,14 +616,17 @@ TSharedPtr<IPlugin> FPluginUtils::CreateAndMountNewPlugin(const FString& PluginN
 			// Generate project files if we happen to be using a project file.
 			if (!FDesktopPlatformModule::Get()->GenerateProjectFiles(FPaths::RootDir(), FPaths::GetProjectFilePath(), GWarn))
 			{
-				FailReason = LOCTEXT("FailedToGenerateProjectFiles", "Failed to generate project files");
+				if (LoadParams.OutFailReason)
+				{
+					*LoadParams.OutFailReason = LOCTEXT("FailedToGenerateProjectFiles", "Failed to generate project files");
+				}
 				bSucceeded = false;
 				break;
 			}
 		}
 
-		// Mount the new plugin
-		NewPlugin = PluginUtils::MountPluginInternal(PluginName, PluginLocation, MountParams, FailReason, /*bIsNewPlugin*/ true);
+		// Load the new plugin
+		NewPlugin = PluginUtils::LoadPluginInternal(PluginName, PluginLocation, PluginFilePath, LoadParams, /*bIsNewPlugin*/ true);
 		if (!NewPlugin)
 		{
 			bSucceeded = false;
@@ -576,22 +666,71 @@ TSharedPtr<IPlugin> FPluginUtils::CreateAndMountNewPlugin(const FString& PluginN
 	return NewPlugin;
 }
 
-TSharedPtr<IPlugin> FPluginUtils::MountPlugin(const FString& PluginName, const FString& PluginLocation, const FMountPluginParams& MountParams, FText& FailReason)
+PRAGMA_DISABLE_DEPRECATION_WARNINGS
+TSharedPtr<IPlugin> FPluginUtils::CreateAndMountNewPlugin(const FString& PluginName, const FString& PluginLocation, const FNewPluginParams& CreationParams, const FMountPluginParams& MountParams, FText& FailReason)
+{
+	FLoadPluginParams LoadParams = PluginUtils::ConvertToLoadPluginParams(MountParams, &FailReason);
+	return CreateAndLoadNewPlugin(PluginName, PluginLocation, CreationParams, LoadParams);
+}
+PRAGMA_ENABLE_DEPRECATION_WARNINGS
+
+PRAGMA_DISABLE_DEPRECATION_WARNINGS
+TSharedPtr<IPlugin> FPluginUtils::CreateAndMountNewPlugin(const FString& PluginName, const FString& PluginLocation, const FNewPluginParamsWithDescriptor& CreationParams, const FMountPluginParams& MountParams, FText& FailReason)
+{
+	FLoadPluginParams LoadParams = PluginUtils::ConvertToLoadPluginParams(MountParams, &FailReason);
+	return CreateAndLoadNewPlugin(PluginName, PluginLocation, CreationParams, LoadParams);
+}
+PRAGMA_ENABLE_DEPRECATION_WARNINGS
+
+TSharedPtr<IPlugin> FPluginUtils::LoadPlugin(const FString& PluginName, const FString& PluginLocation, FLoadPluginParams& LoadParams)
 {
 	// Valide that the uplugin file exists.
 	const FString PluginFilePath = FPluginUtils::GetPluginFilePath(PluginLocation, PluginName, /*bFullPath*/ true);
+
 	if (!FPaths::FileExists(PluginFilePath))
 	{
-		FailReason = FText::Format(LOCTEXT("PluginFileDoesNotExist", "Plugin file does not exist\n{0}"), FText::FromString(PluginFilePath));
+		if (LoadParams.OutFailReason)
+		{
+			*LoadParams.OutFailReason = FText::Format(LOCTEXT("PluginFileDoesNotExist", "Plugin file does not exist\n{0}"), FText::FromString(PluginFilePath));
+		}
 		return nullptr;
 	}
 
-	if (!IsValidPluginName(PluginName, &FailReason))
+	if (!IsValidPluginName(PluginName, LoadParams.OutFailReason))
 	{
 		return nullptr;
 	}
 
-	return PluginUtils::MountPluginInternal(PluginName, PluginLocation, MountParams, FailReason, /*bIsNewPlugin*/ false);
+	return PluginUtils::LoadPluginInternal(PluginName, PluginLocation, PluginFilePath, LoadParams, /*bIsNewPlugin*/ false);
+}
+
+TSharedPtr<IPlugin> FPluginUtils::LoadPlugin(const FString& PluginFileName, FLoadPluginParams& LoadParams)
+{
+	const FString PluginLocation = FPaths::GetPath(FPaths::GetPath(PluginFileName));
+	const FString PluginName = FPaths::GetBaseFilename(PluginFileName);
+
+	return LoadPlugin(PluginName, PluginLocation, LoadParams);
+}
+
+PRAGMA_DISABLE_DEPRECATION_WARNINGS
+TSharedPtr<IPlugin> FPluginUtils::MountPlugin(const FString& PluginName, const FString& PluginLocation, const FMountPluginParams& MountParams, FText& FailReason)
+{
+	FLoadPluginParams LoadParams = PluginUtils::ConvertToLoadPluginParams(MountParams, &FailReason);
+	return LoadPlugin(PluginName, PluginLocation, LoadParams);
+}
+PRAGMA_ENABLE_DEPRECATION_WARNINGS
+
+TSharedPtr<IPlugin> FPluginUtils::FindLoadedPlugin(const FString& PluginDescriptorFileName)
+{
+	TSharedPtr<IPlugin> Plugin = IPluginManager::Get().FindPlugin(FPaths::GetBaseFilename(PluginDescriptorFileName));
+	if (Plugin &&
+		Plugin->IsEnabled() &&
+		FPaths::IsSamePath(Plugin->GetDescriptorFileName(), PluginDescriptorFileName) &&
+		(!Plugin->CanContainContent() || FPackageName::MountPointExists(Plugin->GetMountedAssetPath())))
+	{
+		return Plugin;
+	}
+	return TSharedPtr<IPlugin>();
 }
 
 bool FPluginUtils::UnloadPlugin(const TSharedRef<IPlugin>& Plugin, FText* OutFailReason /*= nullptr*/)
@@ -631,9 +770,9 @@ bool FPluginUtils::UnloadPlugins(const TConstArrayView<TSharedRef<IPlugin>> Plug
 	}
 
 	// Synchronous scan plugins to make sure we find all their assets.
-	IAssetRegistry& AssetRegistry = FModuleManager::LoadModuleChecked<FAssetRegistryModule>(TEXT("AssetRegistry")).Get();
+	IAssetRegistry& AssetRegistry = FModuleManager::LoadModuleChecked<FAssetRegistryModule>(AssetRegistryConstants::ModuleName).Get();
 
-	if (PluginContentMountPoints.IsEmpty() == false)
+	if (!PluginContentMountPoints.IsEmpty())
 	{
 		AssetRegistry.ScanPathsSynchronous(PluginContentMountPoints, /*bForceRescan=*/ true);
 	}
diff --git a/Engine/Plugins/Developer/PluginUtils/Source/PluginUtils/Public/PluginUtils.h b/Engine/Plugins/Developer/PluginUtils/Source/PluginUtils/Public/PluginUtils.h
--- a/Engine/Plugins/Developer/PluginUtils/Source/PluginUtils/Public/PluginUtils.h
+++ b/Engine/Plugins/Developer/PluginUtils/Source/PluginUtils/Public/PluginUtils.h
@@ -114,6 +114,34 @@ public:
 		TArray<FString> TemplateFolders;
 	};
 
+	/**
+	 * Parameters for loading/mounting a plugin
+	 */
+	struct FLoadPluginParams
+	{
+		/** Whether to synchronously scan all assets in the plugin */
+		bool bSynchronousAssetsScan = false;
+
+		/** Whether to select the plugin Content folder (if any) in the content browser */
+		bool bSelectInContentBrowser = false;
+
+		/** Whether to enable the plugin in the current project config */
+		bool bEnablePluginInProject = false;
+
+		/**
+		 * Whether to update the project additional plugin directories (persistently saved in uproject file)
+		 * if the plugin location is not under the engine or project plugin folder
+		 */
+		bool bUpdateProjectPluginSearchPath = false;
+
+		/** Outputs whether the plugin was already loaded */
+		bool bOutAlreadyLoaded =  false;
+
+		/** Outputs the reason the plugin loading failed (if applicable) */
+		FText* OutFailReason = nullptr;
+	};
+
+	struct UE_DEPRECATED(5.1, "FMountPluginParams is deprecated; please use FLoadPluginParams instead") FMountPluginParams;
 	/**
 	 * Parameters for mounting a plugin.
 	 */
@@ -132,6 +160,30 @@ public:
 		bool bSelectInContentBrowser = true;
 	};
 
+	/**
+	 * Helper to create and load a new plugin
+	 * @param PluginName Plugin name
+	 * @param PluginLocation Directory that contains the plugin folder
+	 * @param CreationParams Plugin creation parameters
+	 * @param MountParams Plugin loading parameters
+	 * @return The newly created plugin. If something goes wrong during the creation process, the plugin folder gets deleted and null is returned.
+	 * @note MountParams.OutFailReason outputs the reason the plugin creation or loading failed (if applicable)
+	 * @note Will fail if the plugin already exists
+	 */
+	static TSharedPtr<IPlugin> CreateAndLoadNewPlugin(const FString& PluginName, const FString& PluginLocation, const FNewPluginParams& CreationParams, FLoadPluginParams& LoadParams);
+
+	/**
+	 * Helper to create and load a new plugin
+	 * @param PluginName Plugin name
+	 * @param PluginLocation Directory that contains the plugin folder
+	 * @param CreationParams Plugin creation parameters
+	 * @param LoadParams Plugin loading parameters
+	 * @return The newly created plugin. If something goes wrong during the creation process, the plugin folder gets deleted and null is returned.
+	 * @note MountParams.OutFailReason outputs the reason the plugin creation or loading failed (if applicable)
+	 * @note Will fail if the plugin already exists
+	 */
+	static TSharedPtr<IPlugin> CreateAndLoadNewPlugin(const FString& PluginName, const FString& PluginLocation, const FNewPluginParamsWithDescriptor& CreationParams, FLoadPluginParams& LoadParams);
+
 	/**
 	 * Helper to create and mount a new plugin.
 	 * @param PluginName Plugin name
@@ -142,7 +194,10 @@ public:
 	 * @return The newly created plugin. If something goes wrong during the creation process, the plugin folder gets deleted and null is returned.
 	 * @note Will fail if the plugin already exists
 	 */
+	PRAGMA_DISABLE_DEPRECATION_WARNINGS
+	UE_DEPRECATED(5.1, "CreateAndMountNewPlugin is deprecated; please use CreateAndLoadNewPlugin instead")
 	static TSharedPtr<IPlugin> CreateAndMountNewPlugin(const FString& PluginName, const FString& PluginLocation, const FNewPluginParams& CreationParams, const FMountPluginParams& MountParams, FText& FailReason);
+	PRAGMA_ENABLE_DEPRECATION_WARNINGS
 
 	/**
 	 * Helper to create and mount a new plugin.
@@ -154,18 +209,68 @@ public:
 	 * @return The newly created plugin. If something goes wrong during the creation process, the plugin folder gets deleted and null is returned.
 	 * @note Will fail if the plugin already exists
 	 */
+	PRAGMA_DISABLE_DEPRECATION_WARNINGS
+	UE_DEPRECATED(5.1, "CreateAndMountNewPlugin is deprecated; please use CreateAndLoadNewPlugin instead")
 	static TSharedPtr<IPlugin> CreateAndMountNewPlugin(const FString& PluginName, const FString& PluginLocation, const FNewPluginParamsWithDescriptor& CreationParams, const FMountPluginParams& MountParams, FText& FailReason);
+	PRAGMA_ENABLE_DEPRECATION_WARNINGS
+
+	/**
+	 * Load/mount the specified plugin
+	 * @param PluginName Plugin name
+	 * @param PluginLocation Directory that contains the plugin folder
+	 * @param LoadParams Plugin loading parameters
+	 * @return The loaded plugin or null on failure
+	 */
+	static TSharedPtr<IPlugin> LoadPlugin(const FString& PluginName, const FString& PluginLocation, FLoadPluginParams& LoadParams);
+
+	/**
+	 * Load/mount the specified plugin
+	 * @param PluginName Plugin name
+	 * @param PluginLocation Directory that contains the plugin folder
+	 * @return The loaded plugin or null on failure
+	 */
+	static TSharedPtr<IPlugin> LoadPlugin(const FString& PluginName, const FString& PluginLocation)
+	{
+		FLoadPluginParams Params;
+		return LoadPlugin(PluginName, PluginLocation, Params);
+	}
+
+	/**
+	 * Load/mount the specified plugin
+	 * @param PluginFileName Plugin descriptor file path
+	 * @param LoadParams Plugin loading parameters
+	 * @return The loaded plugin or null on failure
+	 */
+	static TSharedPtr<IPlugin> LoadPlugin(const FString& PluginFileName, FLoadPluginParams& LoadParams);
+
+	/**
+	 * Load/mount the specified plugin
+	 * @param PluginFileName Plugin descriptor file path
+	 * @return The loaded plugin or null on failure
+	 */
+	static TSharedPtr<IPlugin> LoadPlugin(const FString& PluginFileName)
+	{
+		FLoadPluginParams Params;
+		return LoadPlugin(PluginFileName, Params);
+	}
 
 	/**
 	 * Load/mount the specified plugin.
 	 * @param PluginName Plugin name
 	 * @param PluginLocation Directory that contains the plugin folder
-	 * @note the plugin search path will get updated if necessary
 	 * @param MountParams Plugin mounting parameters
 	 * @param FailReason Reason the plugin failed to load
 	 * @return The mounted plugin or null on failure
 	 */
+	PRAGMA_DISABLE_DEPRECATION_WARNINGS
+	UE_DEPRECATED(5.1, "MountPlugin is deprecated; please use LoadPlugin instead")
 	static TSharedPtr<IPlugin> MountPlugin(const FString& PluginName, const FString& PluginLocation, const FMountPluginParams& MountParams, FText& FailReason);
+	PRAGMA_ENABLE_DEPRECATION_WARNINGS
+
+	/**
+	 * Finds a loaded plugin from a plugin descriptor file path
+	 */
+	static TSharedPtr<IPlugin> FindLoadedPlugin(const FString& PluginDescriptorFileName);
 
 	/**
 	 * Unload assets from the specified plugin and unmount it
diff --git a/Engine/Plugins/Editor/PluginBrowser/Source/PluginBrowser/Private/SNewPluginWizard.cpp b/Engine/Plugins/Editor/PluginBrowser/Source/PluginBrowser/Private/SNewPluginWizard.cpp
--- a/Engine/Plugins/Editor/PluginBrowser/Source/PluginBrowser/Private/SNewPluginWizard.cpp
+++ b/Engine/Plugins/Editor/PluginBrowser/Source/PluginBrowser/Private/SNewPluginWizard.cpp
@@ -660,15 +660,16 @@ FReply SNewPluginWizard::OnCreatePluginClicked()
 		CreationParams.Descriptor.bIsBetaVersion = DescriptorData->bIsBetaVersion;
 	}
 
-	FPluginUtils::FMountPluginParams MountParams;
-	MountParams.bEnablePluginInProject = true;
-	MountParams.bUpdateProjectPluginSearchPath = true;
-	MountParams.bSelectInContentBrowser = ShowPluginContentDirectoryCheckBox->IsChecked();
+	FText FailReason;
+	FPluginUtils::FLoadPluginParams LoadParams;
+	LoadParams.bEnablePluginInProject = true;
+	LoadParams.bUpdateProjectPluginSearchPath = true;
+	LoadParams.bSelectInContentBrowser = ShowPluginContentDirectoryCheckBox->IsChecked();
+	LoadParams.OutFailReason = &FailReason;
 
 	Template->CustomizeDescriptorBeforeCreation(CreationParams.Descriptor);
 	
-	FText FailReason;
-	TSharedPtr<IPlugin> NewPlugin = FPluginUtils::CreateAndMountNewPlugin(PluginName, PluginFolderPath, CreationParams, MountParams, FailReason);
+	TSharedPtr<IPlugin> NewPlugin = FPluginUtils::CreateAndLoadNewPlugin(PluginName, PluginFolderPath, CreationParams, LoadParams);
 	const bool bSucceeded = NewPlugin.IsValid();
 
 
diff --git a/Engine/Source/Runtime/Projects/Private/LocalizationDescriptor.cpp b/Engine/Source/Runtime/Projects/Private/LocalizationDescriptor.cpp
--- a/Engine/Source/Runtime/Projects/Private/LocalizationDescriptor.cpp
+++ b/Engine/Source/Runtime/Projects/Private/LocalizationDescriptor.cpp
@@ -75,13 +75,16 @@ FLocalizationTargetDescriptor::FLocalizationTargetDescriptor(FString InName, ELo
 {
 }
 
-bool FLocalizationTargetDescriptor::Read(const FJsonObject& InObject, FText& OutFailReason)
+bool FLocalizationTargetDescriptor::Read(const FJsonObject& InObject, FText* OutFailReason /*= nullptr*/)
 {
 	// Read the target name
 	TSharedPtr<FJsonValue> NameValue = InObject.TryGetField(TEXT("Name"));
 	if (!NameValue.IsValid() || NameValue->Type != EJson::String)
 	{
-		OutFailReason = LOCTEXT("TargetWithoutAName", "Found a 'Localization Target' entry with a missing 'Name' field");
+		if (OutFailReason)
+		{
+			*OutFailReason = LOCTEXT("TargetWithoutAName", "Found a 'Localization Target' entry with a missing 'Name' field");
+		}
 		return false;
 	}
 	Name = NameValue->AsString();
@@ -93,7 +96,10 @@ bool FLocalizationTargetDescriptor::Read(const FJsonObject& InObject, FText& Out
 		LoadingPolicy = ELocalizationTargetDescriptorLoadingPolicy::FromString(*LoadingPolicyValue->AsString());
 		if (LoadingPolicy == ELocalizationTargetDescriptorLoadingPolicy::Max)
 		{
-			OutFailReason = FText::Format(LOCTEXT("TargetWithInvalidLoadingPolicy", "Localization Target entry '{0}' specified an unrecognized target LoadingPolicy '{1}'"), FText::FromString(Name), FText::FromString(LoadingPolicyValue->AsString()));
+			if (OutFailReason)
+			{
+				*OutFailReason = FText::Format(LOCTEXT("TargetWithInvalidLoadingPolicy", "Localization Target entry '{0}' specified an unrecognized target LoadingPolicy '{1}'"), FText::FromString(Name), FText::FromString(LoadingPolicyValue->AsString()));
+			}
 			return false;
 		}
 	}
@@ -101,7 +107,12 @@ bool FLocalizationTargetDescriptor::Read(const FJsonObject& InObject, FText& Out
 	return true;
 }
 
-bool FLocalizationTargetDescriptor::ReadArray(const FJsonObject& InObject, const TCHAR* InName, TArray<FLocalizationTargetDescriptor>& OutTargets, FText& OutFailReason)
+bool FLocalizationTargetDescriptor::Read(const FJsonObject& InObject, FText& OutFailReason)
+{
+	return Read(InObject, &OutFailReason);
+}
+
+bool FLocalizationTargetDescriptor::ReadArray(const FJsonObject& InObject, const TCHAR* InName, TArray<FLocalizationTargetDescriptor>& OutTargets, FText* OutFailReason /*= nullptr*/)
 {
 	bool bResult = true;
 
@@ -125,7 +136,10 @@ bool FLocalizationTargetDescriptor::ReadArray(const FJsonObject& InObject, const
 			}
 			else
 			{
-				OutFailReason = LOCTEXT("TargetWithInvalidTargetsArray", "The 'Localization Targets' array has invalid contents and was not able to be loaded.");
+				if (OutFailReason)
+				{
+					*OutFailReason = LOCTEXT("TargetWithInvalidTargetsArray", "The 'Localization Targets' array has invalid contents and was not able to be loaded.");
+				}
 				bResult = false;
 			}
 		}
@@ -134,6 +148,11 @@ bool FLocalizationTargetDescriptor::ReadArray(const FJsonObject& InObject, const
 	return bResult;
 }
 
+bool FLocalizationTargetDescriptor::ReadArray(const FJsonObject& InObject, const TCHAR* InName, TArray<FLocalizationTargetDescriptor>& OutTargets, FText& OutFailReason)
+{
+	return ReadArray(InObject, InName, OutTargets, &OutFailReason);
+}
+
 void FLocalizationTargetDescriptor::Write(TJsonWriter<>& Writer) const
 {
 	TSharedRef<FJsonObject> DescriptorJsonObject = MakeShared<FJsonObject>();
diff --git a/Engine/Source/Runtime/Projects/Private/ModuleDescriptor.cpp b/Engine/Source/Runtime/Projects/Private/ModuleDescriptor.cpp
--- a/Engine/Source/Runtime/Projects/Private/ModuleDescriptor.cpp
+++ b/Engine/Source/Runtime/Projects/Private/ModuleDescriptor.cpp
@@ -157,13 +157,16 @@ FModuleDescriptor::FModuleDescriptor(const FName InName, EHostType::Type InType,
 {
 }
 
-bool FModuleDescriptor::Read(const FJsonObject& Object, FText& OutFailReason)
+bool FModuleDescriptor::Read(const FJsonObject& Object, FText* OutFailReason /*= nullptr*/)
 {
 	// Read the module name
 	TSharedPtr<FJsonValue> NameValue = Object.TryGetField(TEXT("Name"));
 	if(!NameValue.IsValid() || NameValue->Type != EJson::String)
 	{
-		OutFailReason = LOCTEXT("ModuleWithoutAName", "Found a 'Module' entry with a missing 'Name' field");
+		if (OutFailReason)
+		{
+			*OutFailReason = LOCTEXT("ModuleWithoutAName", "Found a 'Module' entry with a missing 'Name' field");
+		}
 		return false;
 	}
 	Name = FName(*NameValue->AsString());
@@ -172,13 +175,19 @@ bool FModuleDescriptor::Read(const FJsonObject& Object, FText& OutFailReason)
 	TSharedPtr<FJsonValue> TypeValue = Object.TryGetField(TEXT("Type"));
 	if(!TypeValue.IsValid() || TypeValue->Type != EJson::String)
 	{
-		OutFailReason = FText::Format( LOCTEXT( "ModuleWithoutAType", "Found Module entry '{0}' with a missing 'Type' field" ), FText::FromName(Name) );
+		if (OutFailReason)
+		{
+			*OutFailReason = FText::Format( LOCTEXT( "ModuleWithoutAType", "Found Module entry '{0}' with a missing 'Type' field" ), FText::FromName(Name) );
+		}
 		return false;
 	}
 	Type = EHostType::FromString(*TypeValue->AsString());
 	if(Type == EHostType::Max)
 	{
-		OutFailReason = FText::Format( LOCTEXT( "ModuleWithInvalidType", "Module entry '{0}' specified an unrecognized module Type '{1}'" ), FText::FromName(Name), FText::FromString(TypeValue->AsString()) );
+		if (OutFailReason)
+		{
+			*OutFailReason = FText::Format( LOCTEXT( "ModuleWithInvalidType", "Module entry '{0}' specified an unrecognized module Type '{1}'" ), FText::FromName(Name), FText::FromString(TypeValue->AsString()) );
+		}
 		return false;
 	}
 
@@ -189,7 +198,10 @@ bool FModuleDescriptor::Read(const FJsonObject& Object, FText& OutFailReason)
 		LoadingPhase = ELoadingPhase::FromString(*LoadingPhaseValue->AsString());
 		if(LoadingPhase == ELoadingPhase::Max)
 		{
-			OutFailReason = FText::Format( LOCTEXT( "ModuleWithInvalidLoadingPhase", "Module entry '{0}' specified an unrecognized module LoadingPhase '{1}'" ), FText::FromName(Name), FText::FromString(LoadingPhaseValue->AsString()) );
+			if (OutFailReason)
+			{
+				*OutFailReason = FText::Format( LOCTEXT( "ModuleWithInvalidLoadingPhase", "Module entry '{0}' specified an unrecognized module LoadingPhase '{1}'" ), FText::FromName(Name), FText::FromString(LoadingPhaseValue->AsString()) );
+			}
 			return false;
 		}
 	}
@@ -217,7 +229,12 @@ bool FModuleDescriptor::Read(const FJsonObject& Object, FText& OutFailReason)
 	return true;
 }
 
-bool FModuleDescriptor::ReadArray(const FJsonObject& Object, const TCHAR* Name, TArray<FModuleDescriptor>& OutModules, FText& OutFailReason)
+bool FModuleDescriptor::Read(const FJsonObject& Object, FText& OutFailReason)
+{
+	return Read(Object, &OutFailReason);
+}
+
+bool FModuleDescriptor::ReadArray(const FJsonObject& Object, const TCHAR* Name, TArray<FModuleDescriptor>& OutModules, FText* OutFailReason /*= nullptr*/)
 {
 	bool bResult = true;
 
@@ -242,7 +259,10 @@ bool FModuleDescriptor::ReadArray(const FJsonObject& Object, const TCHAR* Name,
 			}
 			else
 			{
-				OutFailReason = LOCTEXT( "ModuleWithInvalidModulesArray", "The 'Modules' array has invalid contents and was not able to be loaded." );
+				if (OutFailReason)
+				{
+					*OutFailReason = LOCTEXT( "ModuleWithInvalidModulesArray", "The 'Modules' array has invalid contents and was not able to be loaded." );
+				}
 				bResult = false;
 			}
 		}
@@ -251,6 +271,11 @@ bool FModuleDescriptor::ReadArray(const FJsonObject& Object, const TCHAR* Name,
 	return bResult;
 }
 
+bool FModuleDescriptor::ReadArray(const FJsonObject& Object, const TCHAR* Name, TArray<FModuleDescriptor>& OutModules, FText& OutFailReason)
+{
+	return ReadArray(Object, Name, OutModules, &OutFailReason);
+}
+
 void FModuleDescriptor::Write(TJsonWriter<>& Writer) const
 {
 	TSharedRef<FJsonObject> ModuleJsonObject = MakeShared<FJsonObject>();
diff --git a/Engine/Source/Runtime/Projects/Private/PluginDescriptor.cpp b/Engine/Source/Runtime/Projects/Private/PluginDescriptor.cpp
--- a/Engine/Source/Runtime/Projects/Private/PluginDescriptor.cpp
+++ b/Engine/Source/Runtime/Projects/Private/PluginDescriptor.cpp
@@ -11,33 +11,42 @@
 
 namespace PluginDescriptor
 {
-	bool ReadFile(const FString& FileName, FString& Text, FText& OutFailReason)
+	bool ReadFile(const FString& FileName, FString& Text, FText* OutFailReason = nullptr)
 	{
 		if (!FFileHelper::LoadFileToString(Text, *FileName))
 		{
-			OutFailReason = FText::Format(LOCTEXT("FailedToLoadDescriptorFile", "Failed to open descriptor file '{0}'"), FText::FromString(FileName));
+			if (OutFailReason)
+			{
+				*OutFailReason = FText::Format(LOCTEXT("FailedToLoadDescriptorFile", "Failed to open descriptor file '{0}'"), FText::FromString(FileName));
+			}
 			return false;
 		}
 		return true;
 	}
 
-	bool WriteFile(const FString& FileName, const FString& Text, FText& OutFailReason)
+	bool WriteFile(const FString& FileName, const FString& Text, FText* OutFailReason = nullptr)
 	{
 		if (!FFileHelper::SaveStringToFile(Text, *FileName))
 		{
-			OutFailReason = FText::Format(LOCTEXT("FailedToWriteDescriptorFile", "Failed to write plugin descriptor file '{0}'. Perhaps the file is Read-Only?"), FText::FromString(FileName));
+			if (OutFailReason)
+			{
+				*OutFailReason = FText::Format(LOCTEXT("FailedToWriteDescriptorFile", "Failed to write plugin descriptor file '{0}'. Perhaps the file is Read-Only?"), FText::FromString(FileName));
+			}
 			return false;
 		}
 		return true;
 	}
 
-	TSharedPtr<FJsonObject> DeserializeJson(const FString& Text, FText& OutFailReason)
+	TSharedPtr<FJsonObject> DeserializeJson(const FString& Text, FText* OutFailReason = nullptr)
 	{
 		TSharedPtr<FJsonObject> JsonObject;
 		TSharedRef<TJsonReader<>> Reader = TJsonReaderFactory<>::Create(Text);
 		if (!FJsonSerializer::Deserialize(Reader, JsonObject) || !JsonObject.IsValid())
 		{
-			OutFailReason = FText::Format(LOCTEXT("FailedToReadDescriptorFile", "Failed to read file. {0}"), FText::FromString(Reader->GetErrorMessage()));
+			if (OutFailReason)
+			{
+				*OutFailReason = FText::Format(LOCTEXT("FailedToReadDescriptorFile", "Failed to read file. {0}"), FText::FromString(Reader->GetErrorMessage()));
+			}
 			return TSharedPtr<FJsonObject>();
 		}
 		return JsonObject;
@@ -82,7 +91,7 @@ FPluginDescriptor::FPluginDescriptor()
 }
 
 
-bool FPluginDescriptor::Load(const FString& FileName, FText& OutFailReason)
+bool FPluginDescriptor::Load(const FString& FileName, FText* OutFailReason /*= nullptr*/)
 {
 #if WITH_EDITOR
 	CachedJson.Reset();
@@ -97,8 +106,12 @@ bool FPluginDescriptor::Load(const FString& FileName, FText& OutFailReason)
 	return false;
 }
 
+bool FPluginDescriptor::Load(const FString& FileName, FText& OutFailReason)
+{
+	return Load(FileName, &OutFailReason);
+}
 
-bool FPluginDescriptor::Read(const FString& Text, FText& OutFailReason)
+bool FPluginDescriptor::Read(const FString& Text, FText* OutFailReason /*= nullptr*/)
 {
 #if WITH_EDITOR
 	CachedJson.Reset();
@@ -122,8 +135,12 @@ bool FPluginDescriptor::Read(const FString& Text, FText& OutFailReason)
 	return false;
 }
 
+bool FPluginDescriptor::Read(const FString& Text, FText& OutFailReason)
+{
+	return Read(Text, &OutFailReason);
+}
 
-bool FPluginDescriptor::Read(const FJsonObject& Object, FText& OutFailReason)
+bool FPluginDescriptor::Read(const FJsonObject& Object, FText* OutFailReason /*= nullptr*/)
 {
 	// Read the file version
 	int32 FileVersionInt32;
@@ -131,7 +148,10 @@ bool FPluginDescriptor::Read(const FJsonObject& Object, FText& OutFailReason)
 	{
 		if(!Object.TryGetNumberField(TEXT("PluginFileVersion"), FileVersionInt32))
 		{
-			OutFailReason = LOCTEXT("InvalidProjectFileVersion", "File does not have a valid 'FileVersion' number.");
+			if (OutFailReason)
+			{
+				*OutFailReason = LOCTEXT("InvalidProjectFileVersion", "File does not have a valid 'FileVersion' number.");
+			}
 			return false;
 		}
 	}
@@ -140,9 +160,12 @@ bool FPluginDescriptor::Read(const FJsonObject& Object, FText& OutFailReason)
 	EPluginDescriptorVersion PluginFileVersion = (EPluginDescriptorVersion)FileVersionInt32;
 	if ((PluginFileVersion <= EPluginDescriptorVersion::Invalid) || (PluginFileVersion > EPluginDescriptorVersion::Latest))
 	{
-		FText ReadVersionText = FText::FromString(FString::Printf(TEXT("%d"), (int32)PluginFileVersion));
-		FText LatestVersionText = FText::FromString(FString::Printf(TEXT("%d"), (int32)EPluginDescriptorVersion::Latest));
-		OutFailReason = FText::Format( LOCTEXT("ProjectFileVersionTooLarge", "File appears to be in a newer version ({0}) of the file format that we can load (max version: {1})."), ReadVersionText, LatestVersionText);
+		if (OutFailReason)
+		{
+			const FText ReadVersionText = FText::FromString(FString::Printf(TEXT("%d"), (int32)PluginFileVersion));
+			const FText LatestVersionText = FText::FromString(FString::Printf(TEXT("%d"), (int32)EPluginDescriptorVersion::Latest));
+			*OutFailReason = FText::Format(LOCTEXT("ProjectFileVersionTooLarge", "File appears to be in a newer version ({0}) of the file format that we can load (max version: {1})."), ReadVersionText, LatestVersionText);
+		}
 		return false;
 	}
 
@@ -218,7 +241,12 @@ bool FPluginDescriptor::Read(const FJsonObject& Object, FText& OutFailReason)
 	return true;
 }
 
-bool FPluginDescriptor::Save(const FString& FileName, FText& OutFailReason) const
+bool FPluginDescriptor::Read(const FJsonObject& Object, FText& OutFailReason)
+{
+	return Read(Object, &OutFailReason);
+}
+
+bool FPluginDescriptor::Save(const FString& FileName, FText* OutFailReason /*= nullptr*/) const
 {
 	// Write the descriptor to text
 	FString Text;
@@ -228,6 +256,11 @@ bool FPluginDescriptor::Save(const FString& FileName, FText& OutFailReason) cons
 	return PluginDescriptor::WriteFile(FileName, Text, OutFailReason);
 }
 
+bool FPluginDescriptor::Save(const FString& FileName, FText& OutFailReason) const
+{
+	return Save(FileName, &OutFailReason);
+}
+
 void FPluginDescriptor::Write(FString& Text) const
 {
 	// Write the contents of the descriptor to a string. Make sure the writer is destroyed so that the contents are flushed to the string.
@@ -389,7 +422,7 @@ void FPluginDescriptor::UpdateJson(FJsonObject& JsonObject) const
 #endif
 }
 
-bool FPluginDescriptor::UpdatePluginFile(const FString& FileName, FText& OutFailReason) const
+bool FPluginDescriptor::UpdatePluginFile(const FString& FileName, FText* OutFailReason /*= nullptr*/) const
 {
 	if (IFileManager::Get().FileExists(*FileName))
 	{
@@ -413,7 +446,10 @@ bool FPluginDescriptor::UpdatePluginFile(const FString& FileName, FText& OutFail
 			TSharedRef<TJsonWriter<>> JsonWriter = TJsonWriterFactory<>::Create(&JsonText);
 			if (!ensure(FJsonSerializer::Serialize(JsonObject.ToSharedRef(), JsonWriter)))
 			{
-				OutFailReason = LOCTEXT("FailedToWriteDescriptor", "Failed to write plugin descriptor content");
+				if (OutFailReason)
+				{
+					*OutFailReason = LOCTEXT("FailedToWriteDescriptor", "Failed to write plugin descriptor content");
+				}
 				return false;
 			}
 		}
@@ -430,6 +466,11 @@ bool FPluginDescriptor::UpdatePluginFile(const FString& FileName, FText& OutFail
 	}
 }
 
+bool FPluginDescriptor::UpdatePluginFile(const FString& FileName, FText& OutFailReason) const
+{
+	return UpdatePluginFile(FileName, &OutFailReason);
+}
+
 bool FPluginDescriptor::SupportsTargetPlatform(const FString& Platform) const
 {
 	if (bHasExplicitPlatforms)
diff --git a/Engine/Source/Runtime/Projects/Private/PluginReferenceDescriptor.cpp b/Engine/Source/Runtime/Projects/Private/PluginReferenceDescriptor.cpp
--- a/Engine/Source/Runtime/Projects/Private/PluginReferenceDescriptor.cpp
+++ b/Engine/Source/Runtime/Projects/Private/PluginReferenceDescriptor.cpp
@@ -117,19 +117,25 @@ bool FPluginReferenceDescriptor::IsSupportedTargetPlatform(const FString& Platfo
 	}
 }
 
-bool FPluginReferenceDescriptor::Read( const FJsonObject& Object, FText& OutFailReason )
+bool FPluginReferenceDescriptor::Read(const FJsonObject& Object, FText* OutFailReason /*= nullptr*/)
 {
 	// Get the name
 	if(!Object.TryGetStringField(TEXT("Name"), Name))
 	{
-		OutFailReason = LOCTEXT("PluginReferenceWithoutName", "Plugin references must have a 'Name' field");
+		if (OutFailReason)
+		{
+			*OutFailReason = LOCTEXT("PluginReferenceWithoutName", "Plugin references must have a 'Name' field");
+		}
 		return false;
 	}
 
 	// Get the enabled field
 	if(!Object.TryGetBoolField(TEXT("Enabled"), bEnabled))
 	{
-		OutFailReason = LOCTEXT("PluginReferenceWithoutEnabled", "Plugin references must have an 'Enabled' field");
+		if (OutFailReason)
+		{
+			*OutFailReason = LOCTEXT("PluginReferenceWithoutEnabled", "Plugin references must have an 'Enabled' field");
+		}
 		return false;
 	}
 
@@ -159,8 +165,12 @@ bool FPluginReferenceDescriptor::Read( const FJsonObject& Object, FText& OutFail
 	return true;
 }
 
+bool FPluginReferenceDescriptor::Read(const FJsonObject& Object, FText& OutFailReason)
+{
+	return Read(Object, &OutFailReason);
+}
 
-bool FPluginReferenceDescriptor::ReadArray( const FJsonObject& Object, const TCHAR* Name, TArray<FPluginReferenceDescriptor>& OutPlugins, FText& OutFailReason )
+bool FPluginReferenceDescriptor::ReadArray(const FJsonObject& Object, const TCHAR* Name, TArray<FPluginReferenceDescriptor>& OutPlugins, FText* OutFailReason /*= nullptr*/)
 {
 	const TArray< TSharedPtr<FJsonValue> > *Array;
 
@@ -187,6 +197,10 @@ bool FPluginReferenceDescriptor::ReadArray( const FJsonObject& Object, const TCH
 	return true;
 }
 
+bool FPluginReferenceDescriptor::ReadArray(const FJsonObject& Object, const TCHAR* Name, TArray<FPluginReferenceDescriptor>& OutPlugins, FText& OutFailReason)
+{
+	return ReadArray(Object, Name, OutPlugins, &OutFailReason);
+}
 
 void FPluginReferenceDescriptor::Write(TJsonWriter<>& Writer) const
 {
diff --git a/Engine/Source/Runtime/Projects/Public/LocalizationDescriptor.h b/Engine/Source/Runtime/Projects/Public/LocalizationDescriptor.h
--- a/Engine/Source/Runtime/Projects/Public/LocalizationDescriptor.h
+++ b/Engine/Source/Runtime/Projects/Public/LocalizationDescriptor.h
@@ -61,9 +61,15 @@ struct PROJECTS_API FLocalizationTargetDescriptor
 	/** Normal constructor */
 	FLocalizationTargetDescriptor(FString InName = FString(), ELocalizationTargetDescriptorLoadingPolicy::Type InLoadingPolicy = ELocalizationTargetDescriptorLoadingPolicy::Never);
 
+	/** Reads a descriptor from the given JSON object */
+	bool Read(const FJsonObject& InObject, FText* OutFailReason = nullptr);
+
 	/** Reads a descriptor from the given JSON object */
 	bool Read(const FJsonObject& InObject, FText& OutFailReason);
 
+	/** Reads an array of targets from the given JSON object */
+	static bool ReadArray(const FJsonObject& InObject, const TCHAR* InName, TArray<FLocalizationTargetDescriptor>& OutTargets, FText* OutFailReason = nullptr);
+
 	/** Reads an array of targets from the given JSON object */
 	static bool ReadArray(const FJsonObject& InObject, const TCHAR* InName, TArray<FLocalizationTargetDescriptor>& OutTargets, FText& OutFailReason);
 
diff --git a/Engine/Source/Runtime/Projects/Public/ModuleDescriptor.h b/Engine/Source/Runtime/Projects/Public/ModuleDescriptor.h
--- a/Engine/Source/Runtime/Projects/Public/ModuleDescriptor.h
+++ b/Engine/Source/Runtime/Projects/Public/ModuleDescriptor.h
@@ -186,9 +186,15 @@ struct PROJECTS_API FModuleDescriptor
 	/** Normal constructor */
 	FModuleDescriptor(const FName InName = NAME_None, EHostType::Type InType = EHostType::Runtime, ELoadingPhase::Type InLoadingPhase = ELoadingPhase::Default);
 
+	/** Reads a descriptor from the given JSON object */
+	bool Read(const FJsonObject& Object, FText* OutFailReason = nullptr);
+
 	/** Reads a descriptor from the given JSON object */
 	bool Read(const FJsonObject& Object, FText& OutFailReason);
 
+	/** Reads an array of modules from the given JSON object */
+	static bool ReadArray(const FJsonObject& Object, const TCHAR* Name, TArray<FModuleDescriptor>& OutModules, FText* OutFailReason = nullptr);
+
 	/** Reads an array of modules from the given JSON object */
 	static bool ReadArray(const FJsonObject& Object, const TCHAR* Name, TArray<FModuleDescriptor>& OutModules, FText& OutFailReason);
 
diff --git a/Engine/Source/Runtime/Projects/Public/PluginDescriptor.h b/Engine/Source/Runtime/Projects/Public/PluginDescriptor.h
--- a/Engine/Source/Runtime/Projects/Public/PluginDescriptor.h
+++ b/Engine/Source/Runtime/Projects/Public/PluginDescriptor.h
@@ -134,15 +134,27 @@ struct PROJECTS_API FPluginDescriptor
 	/** Constructor. */
 	FPluginDescriptor();
 
+	/** Loads the descriptor from the given file. */
+	bool Load(const FString& FileName, FText* OutFailReason = nullptr);
+
 	/** Loads the descriptor from the given file. */
 	bool Load(const FString& FileName, FText& OutFailReason);
 
+	/** Reads the descriptor from the given string */
+	bool Read(const FString& Text, FText* OutFailReason = nullptr);
+
 	/** Reads the descriptor from the given string */
 	bool Read(const FString& Text, FText& OutFailReason);
 
+	/** Reads the descriptor from the given JSON object */
+	bool Read(const FJsonObject& Object, FText* OutFailReason = nullptr);
+
 	/** Reads the descriptor from the given JSON object */
 	bool Read(const FJsonObject& Object, FText& OutFailReason);
 
+	/** Saves the descriptor from the given file. */
+	bool Save(const FString& FileName, FText* OutFailReason = nullptr) const;
+
 	/** Saves the descriptor from the given file. */
 	bool Save(const FString& FileName, FText& OutFailReason) const;
 
@@ -155,6 +167,12 @@ struct PROJECTS_API FPluginDescriptor
 	/** Updates the given json object with values in this descriptor */
 	void UpdateJson(FJsonObject& JsonObject) const;
 
+	/**
+	 * Updates the content of the specified plugin file with values in this descriptor
+	 * (hence preserving json fields that the plugin descriptor doesn't know about)
+	 */
+	bool UpdatePluginFile(const FString& FileName, FText* OutFailReason = nullptr) const;
+
 	/**
 	 * Updates the content of the specified plugin file with values in this descriptor
 	 * (hence preserving json fields that the plugin descriptor doesn't know about)
diff --git a/Engine/Source/Runtime/Projects/Public/PluginReferenceDescriptor.h b/Engine/Source/Runtime/Projects/Public/PluginReferenceDescriptor.h
--- a/Engine/Source/Runtime/Projects/Public/PluginReferenceDescriptor.h
+++ b/Engine/Source/Runtime/Projects/Public/PluginReferenceDescriptor.h
@@ -67,9 +67,15 @@ struct PROJECTS_API FPluginReferenceDescriptor
 	/** Determines if the referenced plugin is supported for the given platform */
 	bool IsSupportedTargetPlatform(const FString& Platform) const;
 
+	/** Reads the descriptor from the given JSON object */
+	bool Read(const FJsonObject& Object, FText* OutFailReason = nullptr);
+
 	/** Reads the descriptor from the given JSON object */
 	bool Read(const FJsonObject& Object, FText& OutFailReason);
 
+	/** Reads an array of modules from the given JSON object */
+	static bool ReadArray(const FJsonObject& Object, const TCHAR* Name, TArray<FPluginReferenceDescriptor>& OutModules, FText* OutFailReason = nullptr);
+
 	/** Reads an array of modules from the given JSON object */
 	static bool ReadArray(const FJsonObject& Object, const TCHAR* Name, TArray<FPluginReferenceDescriptor>& OutModules, FText& OutFailReason);
 
