PR #8683: Fix plugin creation (Contributed by [REDACTED_PERSON])

diff --git a/Engine/Plugins/Developer/PluginUtils/Source/PluginUtils/Private/PluginUtils.cpp b/Engine/Plugins/Developer/PluginUtils/Source/PluginUtils/Private/PluginUtils.cpp
--- a/Engine/Plugins/Developer/PluginUtils/Source/PluginUtils/Private/PluginUtils.cpp
+++ b/Engine/Plugins/Developer/PluginUtils/Source/PluginUtils/Private/PluginUtils.cpp
@@ -310,7 +310,8 @@ namespace PluginUtils
 
 		if (LoadParams.bUpdateProjectPluginSearchPath)
 		{
-			FPluginUtils::AddToPluginSearchPathIfNeeded(PluginLocation, /*bRefreshPlugins=*/true, /*bUpdateProjectFile=*/true);
+			FPluginUtils::AddToPluginSearchPathIfNeeded(PluginLocation, /*bRefreshPlugins=*/false, /*bUpdateProjectFile=*/true);
+			IPluginManager::Get().RefreshPluginsList();
 		}
 		else
 		{
