Make AddToPluginSearchPathIfNeeded skip mod and enterprise plugin directories

diff --git a/Engine/Plugins/Developer/PluginUtils/Source/PluginUtils/Private/PluginUtils.cpp b/Engine/Plugins/Developer/PluginUtils/Source/PluginUtils/Private/PluginUtils.cpp
--- a/Engine/Plugins/Developer/PluginUtils/Source/PluginUtils/Private/PluginUtils.cpp
+++ b/Engine/Plugins/Developer/PluginUtils/Source/PluginUtils/Private/PluginUtils.cpp
@@ -875,49 +875,54 @@ bool FPluginUtils::AddToPluginSearchPathIfNeeded(const FString& Dir, bool bRefre
 
 	const bool bIsEnginePlugin = FPaths::IsUnderDirectory(Dir, FPaths::EnginePluginsDir());
 	const bool bIsProjectPlugin = FPaths::IsUnderDirectory(Dir, FPaths::ProjectPluginsDir());
-	if (!bIsEnginePlugin && !bIsProjectPlugin)
+	const bool bIsModPlugin = FPaths::IsUnderDirectory(Dir, FPaths::ProjectModsDir());
+	const bool bIsEnterprisePlugin = FPaths::IsUnderDirectory(Dir, FPaths::EnterprisePluginsDir());
+	
+	if (bIsEnginePlugin || bIsProjectPlugin || bIsModPlugin || bIsEnterprisePlugin)
 	{
-		if (bUpdateProjectFile)
+		return false;
+	}
+	
+	if (bUpdateProjectFile)
+	{
+		bool bNeedToUpdate = true;
+		for (const FString& AdditionalDir : IProjectManager::Get().GetAdditionalPluginDirectories())
 		{
-			bool bNeedToUpdate = true;
-			for (const FString& AdditionalDir : IProjectManager::Get().GetAdditionalPluginDirectories())
+			if (FPaths::IsUnderDirectory(Dir, AdditionalDir))
 			{
-				if (FPaths::IsUnderDirectory(Dir, AdditionalDir))
-				{
-					bNeedToUpdate = false;
-					break;
-				}
+				bNeedToUpdate = false;
+				break;
 			}
+		}
 
-			if (bNeedToUpdate)
-			{
-				bSearchPathChanged = GameProjectUtils::UpdateAdditionalPluginDirectory(Dir, /*bAdd*/ true);
-			}
+		if (bNeedToUpdate)
+		{
+			bSearchPathChanged = GameProjectUtils::UpdateAdditionalPluginDirectory(Dir, /*bAdd*/ true);
 		}
-		else
+	}
+	else
+	{
+		bool bNeedToUpdate = true;
+		for (const FString& AdditionalDir : IPluginManager::Get().GetAdditionalPluginSearchPaths())
 		{
-			bool bNeedToUpdate = true;
-			for (const FString& AdditionalDir : IPluginManager::Get().GetAdditionalPluginSearchPaths())
-			{
-				if (FPaths::IsUnderDirectory(Dir, AdditionalDir))
-				{
-					bNeedToUpdate = false;
-					break;
-				}
-			}			
-
-			if (bNeedToUpdate)
+			if (FPaths::IsUnderDirectory(Dir, AdditionalDir))
 			{
-				bSearchPathChanged = IPluginManager::Get().AddPluginSearchPath(Dir, /*bShouldRefresh*/ false);
+				bNeedToUpdate = false;
+				break;
 			}
-		}
+		}			
 
-		if (bSearchPathChanged && bRefreshPlugins)
+		if (bNeedToUpdate)
 		{
-			IPluginManager::Get().RefreshPluginsList();
+			bSearchPathChanged = IPluginManager::Get().AddPluginSearchPath(Dir, /*bShouldRefresh*/ false);
 		}
 	}
 
+	if (bSearchPathChanged && bRefreshPlugins)
+	{
+		IPluginManager::Get().RefreshPluginsList();
+	}
+
 	return bSearchPathChanged;
 }
 
