Correct deprecation versions

diff --git a/Engine/Plugins/Developer/PluginUtils/Source/PluginUtils/Public/PluginUtils.h b/Engine/Plugins/Developer/PluginUtils/Source/PluginUtils/Public/PluginUtils.h
--- a/Engine/Plugins/Developer/PluginUtils/Source/PluginUtils/Public/PluginUtils.h
+++ b/Engine/Plugins/Developer/PluginUtils/Source/PluginUtils/Public/PluginUtils.h
@@ -141,7 +141,7 @@ public:
 		FText* OutFailReason = nullptr;
 	};
 
-	struct UE_DEPRECATED(5.1, "FMountPluginParams is deprecated; please use FLoadPluginParams instead") FMountPluginParams;
+	struct UE_DEPRECATED(5.0, "FMountPluginParams is deprecated; please use FLoadPluginParams instead") FMountPluginParams;
 	/**
 	 * Parameters for mounting a plugin.
 	 */
@@ -195,7 +195,7 @@ public:
 	 * @note Will fail if the plugin already exists
 	 */
 	PRAGMA_DISABLE_DEPRECATION_WARNINGS
-	UE_DEPRECATED(5.1, "CreateAndMountNewPlugin is deprecated; please use CreateAndLoadNewPlugin instead")
+	UE_DEPRECATED(5.0, "CreateAndMountNewPlugin is deprecated; please use CreateAndLoadNewPlugin instead")
 	static TSharedPtr<IPlugin> CreateAndMountNewPlugin(const FString& PluginName, const FString& PluginLocation, const FNewPluginParams& CreationParams, const FMountPluginParams& MountParams, FText& FailReason);
 	PRAGMA_ENABLE_DEPRECATION_WARNINGS
 
@@ -210,7 +210,7 @@ public:
 	 * @note Will fail if the plugin already exists
 	 */
 	PRAGMA_DISABLE_DEPRECATION_WARNINGS
-	UE_DEPRECATED(5.1, "CreateAndMountNewPlugin is deprecated; please use CreateAndLoadNewPlugin instead")
+	UE_DEPRECATED(5.0, "CreateAndMountNewPlugin is deprecated; please use CreateAndLoadNewPlugin instead")
 	static TSharedPtr<IPlugin> CreateAndMountNewPlugin(const FString& PluginName, const FString& PluginLocation, const FNewPluginParamsWithDescriptor& CreationParams, const FMountPluginParams& MountParams, FText& FailReason);
 	PRAGMA_ENABLE_DEPRECATION_WARNINGS
 
@@ -263,7 +263,7 @@ public:
 	 * @return The mounted plugin or null on failure
 	 */
 	PRAGMA_DISABLE_DEPRECATION_WARNINGS
-	UE_DEPRECATED(5.1, "MountPlugin is deprecated; please use LoadPlugin instead")
+	UE_DEPRECATED(5.0, "MountPlugin is deprecated; please use LoadPlugin instead")
 	static TSharedPtr<IPlugin> MountPlugin(const FString& PluginName, const FString& PluginLocation, const FMountPluginParams& MountParams, FText& FailReason);
 	PRAGMA_ENABLE_DEPRECATION_WARNINGS
 
diff --git a/Engine/Plugins/Experimental/GeometryScripting/Source/GeometryScriptingCore/Public/GeometryScript/MeshSubdivideFunctions.h b/Engine/Plugins/Experimental/GeometryScripting/Source/GeometryScriptingCore/Public/GeometryScript/MeshSubdivideFunctions.h
--- a/Engine/Plugins/Experimental/GeometryScripting/Source/GeometryScriptingCore/Public/GeometryScript/MeshSubdivideFunctions.h
+++ b/Engine/Plugins/Experimental/GeometryScripting/Source/GeometryScriptingCore/Public/GeometryScript/MeshSubdivideFunctions.h
@@ -26,7 +26,7 @@ class GEOMETRYSCRIPTINGCORE_API UGeometryScriptLibrary_MeshSubdivideFunctions :
 	GENERATED_BODY()
 public:
 
-	UE_DEPRECATED(5.2, "Use 'Apply PN Tessellation' instead; this deprecated version recursively subdivided the mesh NumIterations times while the new version splits every triangle into (TessellationLevel+1)^2 smaller triangles.")
+	UE_DEPRECATED(5.0, "Use 'Apply PN Tessellation' instead; this deprecated version recursively subdivided the mesh NumIterations times while the new version splits every triangle into (TessellationLevel+1)^2 smaller triangles.")
 	UFUNCTION(BlueprintCallable, Category = "GeometryScript|Subdivide", meta = (ScriptMethod, DeprecatedFunction, DeprecationMessage = "Use 'Apply PN Tessellation' instead; this deprecated version recursively subdivided the mesh NumIterations times while the new version splits every triangle into (TessellationLevel+1)^2 smaller triangles."))
 	static UPARAM(DisplayName = "Target Mesh") UDynamicMesh* 
 	ApplyRecursivePNTessellation(
diff --git a/Engine/Source/Developer/ToolMenus/Public/ToolMenuEntry.h b/Engine/Source/Developer/ToolMenus/Public/ToolMenuEntry.h
--- a/Engine/Source/Developer/ToolMenus/Public/ToolMenuEntry.h
+++ b/Engine/Source/Developer/ToolMenus/Public/ToolMenuEntry.h
@@ -193,7 +193,7 @@ public:
 
 	FToolMenuEntryWidgetData WidgetData;
 
-	UE_DEPRECATED(5.1, "Use MakeCustomWidget instead")
+	UE_DEPRECATED(5.0, "Use MakeCustomWidget instead")
 	FNewToolMenuWidget MakeWidget;
 
 	/** Optional delegate that returns a widget to use as this menu entry */
diff --git a/Engine/Source/Runtime/Engine/Classes/Engine/Engine.h b/Engine/Source/Runtime/Engine/Classes/Engine/Engine.h
--- a/Engine/Source/Runtime/Engine/Classes/Engine/Engine.h
+++ b/Engine/Source/Runtime/Engine/Classes/Engine/Engine.h
@@ -2538,7 +2538,7 @@ public:
 	FAudioDeviceHandle GetActiveAudioDevice();
 
 	/** @return whether we currently have more than one local player */
-	UE_DEPRECATED(5.1, "IsSplitScreen was only ever checking if there are more than one local player. Use HasMultipleLocalPlayers instead.")
+	UE_DEPRECATED(5.0, "IsSplitScreen was only ever checking if there are more than one local player. Use HasMultipleLocalPlayers instead.")
 	virtual bool IsSplitScreen(UWorld *InWorld);
 
 	/** @returns whether there are currently multiple local players in the given world */
diff --git a/Engine/Source/Runtime/Engine/Classes/Engine/NetConnection.h b/Engine/Source/Runtime/Engine/Classes/Engine/NetConnection.h
--- a/Engine/Source/Runtime/Engine/Classes/Engine/NetConnection.h
+++ b/Engine/Source/Runtime/Engine/Classes/Engine/NetConnection.h
@@ -1558,7 +1558,7 @@ private:
 	/** Out of order packet tracking/correction */
 
 	/** Stat tracking for the total number of out of order packets, for this connection */
-	UE_DEPRECATED(5.1, "Use GetTotalOutOfOrderPackets instead.")
+	UE_DEPRECATED(5.0, "Use GetTotalOutOfOrderPackets instead.")
 	int32 TotalOutOfOrderPackets = 0;
 
 	/** Stat tracking for the total number of out of order packets lost */
diff --git a/Engine/Source/Runtime/Engine/Classes/Engine/NetDriver.h b/Engine/Source/Runtime/Engine/Classes/Engine/NetDriver.h
--- a/Engine/Source/Runtime/Engine/Classes/Engine/NetDriver.h
+++ b/Engine/Source/Runtime/Engine/Classes/Engine/NetDriver.h
@@ -959,9 +959,9 @@ public:
 	uint32						OutPacketsLost;
 	/** Total packets lost that have been sent by the server since the net driver's creation  */
 	uint32						OutTotalPacketsLost;
-	UE_DEPRECATED(5.1, "Use GetTotalOutOfOrderPackets instead.")
+	UE_DEPRECATED(5.0, "Use GetTotalOutOfOrderPackets instead.")
 	uint32						InOutOfOrderPackets = 0;
-	UE_DEPRECATED(5.1, "OutOutOfOrderPackets is not updated and is now deprecated.")
+	UE_DEPRECATED(5.0, "OutOutOfOrderPackets is not updated and is now deprecated.")
 	uint32						OutOutOfOrderPackets = 0;
 	/** Tracks the total number of voice packets sent */
 	uint32						VoicePacketsSent;
diff --git a/Engine/Source/Runtime/Engine/Classes/Kismet/KismetSystemLibrary.h b/Engine/Source/Runtime/Engine/Classes/Kismet/KismetSystemLibrary.h
--- a/Engine/Source/Runtime/Engine/Classes/Kismet/KismetSystemLibrary.h
+++ b/Engine/Source/Runtime/Engine/Classes/Kismet/KismetSystemLibrary.h
@@ -209,7 +209,7 @@ class ENGINE_API UKismetSystemLibrary : public UBlueprintFunctionLibrary
 	static bool IsStandalone(const UObject* WorldContextObject);
 	
 	/** Returns whether we currently have more than one local player. */
-	UE_DEPRECATED(5.1, "IsSplitScreen was only ever checking if there are more than one local player. Use HasMultipleLocalPlayers instead.")
+	UE_DEPRECATED(5.0, "IsSplitScreen was only ever checking if there are more than one local player. Use HasMultipleLocalPlayers instead.")
 	UFUNCTION(BlueprintPure, Category = "Viewport", meta = (WorldContext = "WorldContextObject", DeprecatedFunction, DeprecationMessage = "Use HasMultipleLocalPlayers instead"))
 	static bool IsSplitScreen(const UObject* WorldContextObject);
 
diff --git a/Engine/Source/Runtime/Engine/Public/HDRLoader.h b/Engine/Source/Runtime/Engine/Public/HDRLoader.h
--- a/Engine/Source/Runtime/Engine/Public/HDRLoader.h
+++ b/Engine/Source/Runtime/Engine/Public/HDRLoader.h
@@ -9,7 +9,7 @@
 // http://paulbourke.net/dataformats/pic/
 
 /** To load the HDR file image format. Does not support all possible types HDR formats (e.g. xyze is not supported) */
-class UE_DEPRECATED(5.1, "HDR Loader as been replaced by the FHdrImageWrapper which can be found in the image wrapper module.") FHDRLoadHelper;
+class UE_DEPRECATED(5.0, "HDR Loader as been replaced by the FHdrImageWrapper which can be found in the image wrapper module.") FHDRLoadHelper;
 class FHDRLoadHelper
 {
 public:
diff --git a/Engine/Source/Runtime/RHI/Public/RHI.h b/Engine/Source/Runtime/RHI/Public/RHI.h
--- a/Engine/Source/Runtime/RHI/Public/RHI.h
+++ b/Engine/Source/Runtime/RHI/Public/RHI.h
@@ -430,7 +430,7 @@ extern RHI_API bool GSupportsEfficientAsyncCompute;
 extern RHI_API bool GSupportsParallelOcclusionQueries;
 
 /** true if the RHI supports aliasing of transient resources */
-UE_DEPRECATED(5.1, "GSupportsTransientResourceAliasing has been deprecated. Transient resource aliasing is handled through IRHITransientResourceAllocator.")
+UE_DEPRECATED(5.0, "GSupportsTransientResourceAliasing has been deprecated. Transient resource aliasing is handled through IRHITransientResourceAllocator.")
 extern RHI_API bool GSupportsTransientResourceAliasing;
 
 /** true if the RHI requires a valid RT bound during UAV scatter operation inside the pixel shader */
diff --git a/Engine/Source/Runtime/RHI/Public/RHIResources.h b/Engine/Source/Runtime/RHI/Public/RHIResources.h
--- a/Engine/Source/Runtime/RHI/Public/RHIResources.h
+++ b/Engine/Source/Runtime/RHI/Public/RHIResources.h
@@ -54,7 +54,7 @@ enum class EClearBinding;
 class RHI_API FRHIResource
 {
 public:
-	UE_DEPRECATED(5.1, "FRHIResource(bool) is deprecated, please use FRHIResource(ERHIResourceType)")
+	UE_DEPRECATED(5.0, "FRHIResource(bool) is deprecated, please use FRHIResource(ERHIResourceType)")
 	FRHIResource(bool InbDoNotDeferDelete=false)
 		: ResourceType(RRT_None)
 	{
