Fixing source control issues when PLUGIN_NAME is used in uassets in plugin templates.

diff --git a/Engine/Plugins/Developer/PluginUtils/Source/PluginUtils/Private/PluginUtils.cpp b/Engine/Plugins/Developer/PluginUtils/Source/PluginUtils/Private/PluginUtils.cpp
--- a/Engine/Plugins/Developer/PluginUtils/Source/PluginUtils/Private/PluginUtils.cpp
+++ b/Engine/Plugins/Developer/PluginUtils/Source/PluginUtils/Private/PluginUtils.cpp
@@ -215,7 +215,7 @@ namespace PluginUtils
 		return PlatformFile.IterateDirectoryRecursively(*SourceDir, CopyFilesAndDirs);
 	}
 
-	void FixupPluginTemplateAssets(const FString& PluginName)
+	void FixupPluginTemplateAssets(const FString& PluginName, TMap<FString, FString>& OutModifiedAssetPaths)
 	{
 		IPlatformFile& PlatformFile = FPlatformFileManager::Get().GetPlatformFile();
 
@@ -256,7 +256,7 @@ namespace PluginUtils
 			 * Fixes up any assets that contain the PLUGIN_NAME text macro, since those need to be renamed by the engine for the change to
 			 * stick (as opposed to just renaming the file)
 			 */
-			void PerformFixup()
+			void PerformFixup(TMap<FString, FString>& OutModifiedAssetPaths)
 			{
 				TArray<FAssetRenameData> AssetRenameData;
 
@@ -265,6 +265,7 @@ namespace PluginUtils
 					IAssetRegistry& AssetRegistry = FModuleManager::LoadModuleChecked<FAssetRegistryModule>(AssetRegistryConstants::ModuleName).Get();
 					AssetRegistry.ScanFilesSynchronous(FilesToScan);
 
+					TMap<TWeakObjectPtr<UObject>, FString> AssetToOriginalFilePathMap;
 					for (const FString& File : FilesToScan)
 					{
 						TArray<FAssetData> Assets;
@@ -279,7 +280,7 @@ namespace PluginUtils
 						{
 							const FString AssetName = Asset.AssetName.ToString().Replace(*PLUGIN_NAME, *PluginName, ESearchCase::CaseSensitive);
 							const FString AssetPath = Asset.PackagePath.ToString().Replace(*PLUGIN_NAME, *PluginName, ESearchCase::CaseSensitive);
-
+							AssetToOriginalFilePathMap.Add(Asset.GetAsset(), File);
 							FAssetRenameData RenameData(Asset.GetAsset(), AssetPath, AssetName);
 
 							AssetRenameData.Add(RenameData);
@@ -290,6 +291,24 @@ namespace PluginUtils
 					{
 						FAssetToolsModule& AssetToolsModule = FModuleManager::LoadModuleChecked<FAssetToolsModule>("AssetTools");
 						AssetToolsModule.Get().RenameAssetsWithDialog(AssetRenameData);
+
+						for(const FAssetRenameData& RenamedAsset : AssetRenameData)
+						{
+							if (RenamedAsset.Asset.IsValid())
+							{
+								if (FString* OldPath = AssetToOriginalFilePathMap.Find(RenamedAsset.Asset))
+								{
+									if (const UPackage* Package = RenamedAsset.Asset->GetPackage())
+									{
+										FString PackageFilename;
+										if (FPackageName::TryConvertLongPackageNameToFilename(Package->GetName(), PackageFilename, Package->ContainsMap() ? FPackageName::GetMapPackageExtension() : FPackageName::GetAssetPackageExtension()))
+										{
+											OutModifiedAssetPaths.Add(*OldPath, PackageFilename);
+										}
+									}
+								}
+							}
+						}
 					}
 				}
 			}
@@ -300,7 +319,7 @@ namespace PluginUtils
 			const FString PluginBaseDir = Plugin->GetBaseDir();
 			FFixupPluginAssets FixupPluginAssets(PlatformFile, PluginName, PluginBaseDir);
 			PlatformFile.IterateDirectoryRecursively(*PluginBaseDir, FixupPluginAssets);
-			FixupPluginAssets.PerformFixup();
+			FixupPluginAssets.PerformFixup(OutModifiedAssetPaths);
 		}
 	}
 
@@ -638,7 +657,16 @@ TSharedPtr<IPlugin> FPluginUtils::CreateAndLoadNewPlugin(const FString& PluginNa
 		if (CreationParams.Descriptor.bCanContainContent)
 		{	
 			GWarn->BeginSlowTask(LOCTEXT("LoadingContent", "Loading Content..."), /*ShowProgressDialog*/ true, /*bShowCancelButton*/ false);
-			PluginUtils::FixupPluginTemplateAssets(PluginName);
+			TMap<FString, FString> ModifiedAssetPaths;
+			PluginUtils::FixupPluginTemplateAssets(PluginName, ModifiedAssetPaths);
+
+			for (FString& CopiedFilePath : NewFilePaths)
+			{
+				if (FString* NewPath = ModifiedAssetPaths.Find(CopiedFilePath))
+				{
+					CopiedFilePath = *NewPath;
+				}
+			}
 			GWarn->EndSlowTask();
 		}
 
