Adding uplugin to NameReplacementFileTypes for the plugin creation wizard.

diff --git a/Engine/Plugins/Developer/PluginUtils/Source/PluginUtils/Private/PluginUtils.cpp b/Engine/Plugins/Developer/PluginUtils/Source/PluginUtils/Private/PluginUtils.cpp
--- a/Engine/Plugins/Developer/PluginUtils/Source/PluginUtils/Private/PluginUtils.cpp
+++ b/Engine/Plugins/Developer/PluginUtils/Source/PluginUtils/Private/PluginUtils.cpp
@@ -104,6 +104,7 @@ namespace PluginUtils
 				NameReplacementFileTypes.Add(TEXT("cpp"));
 				NameReplacementFileTypes.Add(TEXT("h"));
 				NameReplacementFileTypes.Add(TEXT("vcxproj"));
+				NameReplacementFileTypes.Add(TEXT("uplugin"));
 
 				// Which file types do we want to ignore
 				IgnoredFileTypes.Add(TEXT("opensdf"));
