Removing assumption that the uplugin file lives in a folder of the same name in FPluginUtils::LoadPlugin

diff --git a/Engine/Plugins/Developer/PluginUtils/Source/PluginUtils/Private/PluginUtils.cpp b/Engine/Plugins/Developer/PluginUtils/Source/PluginUtils/Private/PluginUtils.cpp
--- a/Engine/Plugins/Developer/PluginUtils/Source/PluginUtils/Private/PluginUtils.cpp
+++ b/Engine/Plugins/Developer/PluginUtils/Source/PluginUtils/Private/PluginUtils.cpp
@@ -714,32 +714,29 @@ PRAGMA_ENABLE_DEPRECATION_WARNINGS
 
 TSharedPtr<IPlugin> FPluginUtils::LoadPlugin(const FString& PluginName, const FString& PluginLocation, FLoadPluginParams& LoadParams)
 {
-	// Valide that the uplugin file exists.
-	const FString PluginFilePath = FPluginUtils::GetPluginFilePath(PluginLocation, PluginName, /*bFullPath*/ true);
+	return LoadPlugin(FPluginUtils::GetPluginFilePath(PluginLocation, PluginName), LoadParams);
+}
 
-	if (!FPaths::FileExists(PluginFilePath))
+TSharedPtr<IPlugin> FPluginUtils::LoadPlugin(const FString& PluginFileName, FLoadPluginParams& LoadParams)
+{
+	// Valide that the uplugin file exists.
+	if (!FPaths::FileExists(PluginFileName))
 	{
 		if (LoadParams.OutFailReason)
 		{
-			*LoadParams.OutFailReason = FText::Format(LOCTEXT("PluginFileDoesNotExist", "Plugin file does not exist\n{0}"), FText::FromString(PluginFilePath));
+			*LoadParams.OutFailReason = FText::Format(LOCTEXT("PluginFileDoesNotExist", "Plugin file does not exist\n{0}"), FText::FromString(FPaths::ConvertRelativePathToFull(PluginFileName)));
 		}
 		return nullptr;
 	}
 
+	const FString PluginName = FPaths::GetBaseFilename(PluginFileName);
 	if (!IsValidPluginName(PluginName, LoadParams.OutFailReason))
 	{
 		return nullptr;
 	}
 
-	return PluginUtils::LoadPluginInternal(PluginName, PluginLocation, PluginFilePath, LoadParams, /*bIsNewPlugin*/ false);
-}
-
-TSharedPtr<IPlugin> FPluginUtils::LoadPlugin(const FString& PluginFileName, FLoadPluginParams& LoadParams)
-{
 	const FString PluginLocation = FPaths::GetPath(FPaths::GetPath(PluginFileName));
-	const FString PluginName = FPaths::GetBaseFilename(PluginFileName);
-
-	return LoadPlugin(PluginName, PluginLocation, LoadParams);
+	return PluginUtils::LoadPluginInternal(PluginName, PluginLocation, PluginFileName, LoadParams, /*bIsNewPlugin*/ false);
 }
 
 PRAGMA_DISABLE_DEPRECATION_WARNINGS
