Fixing mismatched relative/absolute paths for path comparison in plugin utils

diff --git a/Engine/Plugins/Developer/PluginUtils/Source/PluginUtils/Private/PluginUtils.cpp b/Engine/Plugins/Developer/PluginUtils/Source/PluginUtils/Private/PluginUtils.cpp
--- a/Engine/Plugins/Developer/PluginUtils/Source/PluginUtils/Private/PluginUtils.cpp
+++ b/Engine/Plugins/Developer/PluginUtils/Source/PluginUtils/Private/PluginUtils.cpp
@@ -304,7 +304,9 @@ namespace PluginUtils
 										FString PackageFilename;
 										if (FPackageName::TryConvertLongPackageNameToFilename(Package->GetName(), PackageFilename, Package->ContainsMap() ? FPackageName::GetMapPackageExtension() : FPackageName::GetAssetPackageExtension()))
 										{
-											OutModifiedAssetPaths.Add(*OldPath, PackageFilename);
+											FString FullOldPath = FPaths::ConvertRelativePathToFull(*OldPath);
+											FPaths::RemoveDuplicateSlashes(FullOldPath);
+											OutModifiedAssetPaths.Add(FullOldPath, PackageFilename);
 										}
 									}
 								}
@@ -663,7 +665,9 @@ TSharedPtr<IPlugin> FPluginUtils::CreateAndLoadNewPlugin(const FString& PluginNa
 
 			for (FString& CopiedFilePath : NewFilePaths)
 			{
-				if (FString* NewPath = ModifiedAssetPaths.Find(CopiedFilePath))
+				FString FullCopiedPath = FPaths::ConvertRelativePathToFull(CopiedFilePath);
+				FPaths::RemoveDuplicateSlashes(FullCopiedPath);
+				if (FString* NewPath = ModifiedAssetPaths.Find(FullCopiedPath))
 				{
 					CopiedFilePath = *NewPath;
 				}
