Proper unloading of gfp assets in editor (need to deal with assets opened in editors and standalone packages)

diff --git a/Engine/Plugins/Developer/PluginUtils/Source/PluginUtils/Private/PluginUtils.cpp b/Engine/Plugins/Developer/PluginUtils/Source/PluginUtils/Private/PluginUtils.cpp
--- a/Engine/Plugins/Developer/PluginUtils/Source/PluginUtils/Private/PluginUtils.cpp
+++ b/Engine/Plugins/Developer/PluginUtils/Source/PluginUtils/Private/PluginUtils.cpp
@@ -785,6 +785,86 @@ bool FPluginUtils::UnloadPlugins(const TConstArrayView<TSharedRef<IPlugin>> Plug
 		return true;
 	}
 
+	{
+		FText ErrorMsg;
+		if (!UnloadPluginsAssets(Plugins, &ErrorMsg))
+		{
+			// If some assets fail to unload, log an error, but unmount the plugins anyway
+			UE_LOG(LogPluginUtils, Error, TEXT("Failed to unload some assets prior to unmounting plugins\n%s"), *ErrorMsg.ToString());
+		}
+	}
+
+	// Unmount the plugins
+	//
+	bool bSuccess = true;
+	{
+		FTextBuilder ErrorBuilder;
+		bool bPluginUnmounted = false;
+
+		for (const TSharedRef<IPlugin>& Plugin : Plugins)
+		{
+			if (Plugin->IsEnabled())
+			{
+				bPluginUnmounted = true;
+
+				FText FailReason;
+				if (!IPluginManager::Get().UnmountExplicitlyLoadedPlugin(Plugin->GetName(), &FailReason))
+				{
+					UE_LOG(LogPluginUtils, Error, TEXT("Plugin %s cannot be unloaded: %s"), *Plugin->GetName(), *FailReason.ToString());
+					ErrorBuilder.AppendLine(FailReason);
+					bSuccess = false;
+				}
+			}
+		}
+
+		if (bPluginUnmounted)
+		{
+			IPluginManager::Get().RefreshPluginsList();
+		}
+
+		if (!bSuccess && OutFailReason)
+		{
+			*OutFailReason = ErrorBuilder.ToText();
+		}
+	}
+	return bSuccess;
+}
+
+bool FPluginUtils::UnloadPlugins(const TConstArrayView<FString> PluginNames, FText* OutFailReason /*= nullptr*/)
+{
+	TArray<TSharedRef<IPlugin>> Plugins;
+	Plugins.Reserve(PluginNames.Num());
+	for (const FString& PluginName : PluginNames)
+	{
+		if (TSharedPtr<IPlugin> Plugin = IPluginManager::Get().FindPlugin(PluginName))
+		{
+			Plugins.Add(Plugin.ToSharedRef());
+		}
+	}
+	return UnloadPlugins(Plugins, OutFailReason);
+}
+
+bool FPluginUtils::UnloadPluginAssets(const TSharedRef<IPlugin>& Plugin, FText* OutFailReason /*= nullptr*/)
+{
+	return UnloadPluginsAssets({ Plugin }, OutFailReason);
+}
+
+bool FPluginUtils::UnloadPluginAssets(const FString& PluginName, FText* OutFailReason /*= nullptr*/)
+{
+	if (TSharedPtr<IPlugin> Plugin = IPluginManager::Get().FindPlugin(PluginName))
+	{
+		return UnloadPluginAssets(Plugin.ToSharedRef(), OutFailReason);
+	}
+	return true;
+}
+
+bool FPluginUtils::UnloadPluginsAssets(const TConstArrayView<TSharedRef<IPlugin>> Plugins, FText* OutFailReason /*= nullptr*/)
+{
+	if (Plugins.Num() == 0)
+	{
+		return true;
+	}
+
 	TArray<FString> PluginContentMountPoints;
 	PluginContentMountPoints.Reserve(Plugins.Num());
 
@@ -838,54 +918,24 @@ bool FPluginUtils::UnloadPlugins(const TConstArrayView<TSharedRef<IPlugin>> Plug
 				FText ErrorMsg;
 				UPackageTools::UnloadPackages(PackagesToUnload.Array(), ErrorMsg, /*bUnloadDirtyPackages=*/true);
 
-				// Log an error if some packages fail to unload, but unmount the plugins anyway.
 				// @note UnloadPackages returned bool indicates whether some packages were unloaded
 				// To tell whether all packages were successfully unloaded we must check the ErrorMsg output param
 				if (!ErrorMsg.IsEmpty())
 				{
-					UE_LOG(LogPluginUtils, Error, TEXT("%s"), *ErrorMsg.ToString());
+					if (OutFailReason)
+					{
+						*OutFailReason = MoveTemp(ErrorMsg);
+					}
+					return false;
 				}
 			}
 		}
 	}
 
-	// Unmount the plugins
-	//
-	bool bSuccess = true;
-	{
-		FTextBuilder ErrorBuilder;
-		bool bPluginUnmounted = false;
-
-		for (const TSharedRef<IPlugin>& Plugin : Plugins)
-		{
-			if (Plugin->IsEnabled())
-			{
-				bPluginUnmounted = true;
-
-				FText FailReason;
-				if (!IPluginManager::Get().UnmountExplicitlyLoadedPlugin(Plugin->GetName(), &FailReason))
-				{
-					UE_LOG(LogPluginUtils, Error, TEXT("Plugin %s cannot be unloaded: %s"), *Plugin->GetName(), *FailReason.ToString());
-					ErrorBuilder.AppendLine(FailReason);
-					bSuccess = false;
-				}
-			}
-		}
-
-		if (bPluginUnmounted)
-		{
-			IPluginManager::Get().RefreshPluginsList();
-		}
-
-		if (!bSuccess && OutFailReason)
-		{
-			*OutFailReason = ErrorBuilder.ToText();
-		}
-	}
-	return bSuccess;
+	return true;
 }
 
-bool FPluginUtils::UnloadPlugins(const TConstArrayView<FString> PluginNames, FText* OutFailReason /*= nullptr*/)
+bool FPluginUtils::UnloadPluginsAssets(const TConstArrayView<FString> PluginNames, FText* OutFailReason /*= nullptr*/)
 {
 	TArray<TSharedRef<IPlugin>> Plugins;
 	Plugins.Reserve(PluginNames.Num());
@@ -896,7 +946,7 @@ bool FPluginUtils::UnloadPlugins(const TConstArrayView<FString> PluginNames, FTe
 			Plugins.Add(Plugin.ToSharedRef());
 		}
 	}
-	return UnloadPlugins(Plugins, OutFailReason);
+	return UnloadPluginsAssets(Plugins, OutFailReason);
 }
 
 bool FPluginUtils::AddToPluginSearchPathIfNeeded(const FString& Dir, bool bRefreshPlugins, bool bUpdateProjectFile)
diff --git a/Engine/Plugins/Developer/PluginUtils/Source/PluginUtils/Public/PluginUtils.h b/Engine/Plugins/Developer/PluginUtils/Source/PluginUtils/Public/PluginUtils.h
--- a/Engine/Plugins/Developer/PluginUtils/Source/PluginUtils/Public/PluginUtils.h
+++ b/Engine/Plugins/Developer/PluginUtils/Source/PluginUtils/Public/PluginUtils.h
@@ -312,6 +312,42 @@ public:
 	 */
 	static bool UnloadPlugins(const TConstArrayView<FString> PluginNames, FText* OutFailReason = nullptr);
 
+	/**
+	 * Unload assets from the specified plugin but does not unmount it
+	 * @warning Dirty assets that need to be saved will be unloaded anyway
+	 * @param Plugin Plugin to unload assets from
+	 * @param OutFailReason Outputs the reason of the failure if any
+	 * @return Whether plugin assets were successfully unloaded
+	 */
+	static bool UnloadPluginAssets(const TSharedRef<IPlugin>& Plugin, FText* OutFailReason = nullptr);
+
+	/**
+	 * Unload assets from the specified plugin but does not unmount it
+	 * @warning Dirty assets that need to be saved will be unloaded anyway
+	 * @param PluginName Name of the plugin to unload assets from
+	 * @param OutFailReason Outputs the reason of the failure if any
+	 * @return Whether plugin assets were successfully unloaded
+	 */
+	static bool UnloadPluginAssets(const FString& PluginName, FText* OutFailReason = nullptr);
+
+	/**
+	 * Unload assets from the specified plugins but does not unmount them
+	 * @warning Dirty assets that need to be saved will be unloaded anyway
+	 * @param Plugin Plugins to unload assets from
+	 * @param OutFailReason Outputs the reason of the failure if any
+	 * @return Whether plugin assets were successfully unloaded
+	 */
+	static bool UnloadPluginsAssets(const TConstArrayView<TSharedRef<IPlugin>> Plugins, FText* OutFailReason = nullptr);
+
+	/**
+	 * Unload assets from the specified plugin but does not unmount them
+	 * @warning Dirty assets that need to be saved will be unloaded anyway
+	 * @param PluginNames Names of the plugins to unload assets from
+	 * @param OutFailReason Outputs the reason of the failure if any
+	 * @return Whether plugin assets were successfully unloaded
+	 */
+	static bool UnloadPluginsAssets(const TConstArrayView<FString> PluginNames, FText* OutFailReason = nullptr);
+
 	/**
 	 * Adds a directory to the list of paths that are recursively searched for plugins, 
 	 * if that directory isn't already under the search paths.
diff --git a/Engine/Plugins/Experimental/GameFeatures/GameFeatures.uplugin b/Engine/Plugins/Experimental/GameFeatures/GameFeatures.uplugin
--- a/Engine/Plugins/Experimental/GameFeatures/GameFeatures.uplugin
+++ b/Engine/Plugins/Experimental/GameFeatures/GameFeatures.uplugin
@@ -38,6 +38,11 @@
 		{
 			"Name": "AssetReferenceRestrictions",
 			"Enabled": true
+		},
+		{
+			"Name": "PluginUtils",
+			"Enabled": true,
+			"TargetAllowList": [ "Editor" ]
 		}
 	]
 }
diff --git a/Engine/Plugins/Experimental/GameFeatures/Source/GameFeatures/GameFeatures.Build.cs b/Engine/Plugins/Experimental/GameFeatures/Source/GameFeatures/GameFeatures.Build.cs
--- a/Engine/Plugins/Experimental/GameFeatures/Source/GameFeatures/GameFeatures.Build.cs
+++ b/Engine/Plugins/Experimental/GameFeatures/Source/GameFeatures/GameFeatures.Build.cs
@@ -29,6 +29,16 @@ namespace UnrealBuildTool.Rules
 					"Projects",
 				}
 			);
+
+			if (Target.bBuildEditor)
+			{
+				PrivateDependencyModuleNames.AddRange(
+					new string[]
+					{
+						"PluginUtils",
+					}
+				);
+			}
 		}
 	}
 }
diff --git a/Engine/Plugins/Experimental/GameFeatures/Source/GameFeatures/Private/GameFeaturePluginStateMachine.cpp b/Engine/Plugins/Experimental/GameFeatures/Source/GameFeatures/Private/GameFeaturePluginStateMachine.cpp
--- a/Engine/Plugins/Experimental/GameFeatures/Source/GameFeatures/Private/GameFeaturePluginStateMachine.cpp
+++ b/Engine/Plugins/Experimental/GameFeatures/Source/GameFeatures/Private/GameFeaturePluginStateMachine.cpp
@@ -23,6 +23,10 @@
 #include "GameFeaturesProjectPolicies.h"
 #include "Containers/Ticker.h"
 
+#if WITH_EDITOR
+#include "PluginUtils.h"
+#endif //if WITH_EDITOR
+
 namespace UE::GameFeatures
 {
 	static const FString StateMachineErrorNamespace(TEXT("GameFeaturePlugin.StateMachine."));
@@ -1191,6 +1195,12 @@ struct FGameFeaturePluginState_Unregistering : public FGameFeaturePluginState
 
 		StateProperties.GameFeatureData = nullptr;
 
+#if WITH_EDITOR
+		// This will properly unload any plugin asset that could be opened in the editor
+		// and ensure standalone packages get unloaded as well
+		verify(FPluginUtils::UnloadPluginAssets(StateProperties.PluginName));
+#endif //if WITH_EDITOR
+
 		bRequestedGC = true;
 		GarbageCollectAndUpdateStateMachineDeferred();
 	}
@@ -1286,7 +1296,9 @@ struct FGameFeaturePluginState_Unloading : public FGameFeaturePluginState
 	{
 		if (bRequestedGC)
 		{
+#if !WITH_EDITOR // Disabled in editor since it's likely to report unloaded assets because of standalone packages
 			UE::GameFeatures::VerifyAssetsUnloaded(StateProperties.PluginName, true);
+#endif //if !WITH_EDITOR
 
 			StateStatus.SetTransition(EGameFeaturePluginState::Registered);
 			return;
