Fixing internal code modules being stored in external module location when creating internal plugins through the plugin wizard.

diff --git a/Engine/Plugins/Developer/PluginUtils/Source/PluginUtils/Private/PluginUtils.cpp b/Engine/Plugins/Developer/PluginUtils/Source/PluginUtils/Private/PluginUtils.cpp
--- a/Engine/Plugins/Developer/PluginUtils/Source/PluginUtils/Private/PluginUtils.cpp
+++ b/Engine/Plugins/Developer/PluginUtils/Source/PluginUtils/Private/PluginUtils.cpp
@@ -622,7 +622,14 @@ TSharedPtr<IPlugin> FPluginUtils::CreateAndLoadNewPlugin(const FString& PluginNa
 		if (CreationParams.Descriptor.Modules.Num() > 0)
 		{
 			const FString ProjectFileName = IFileManager::Get().ConvertToAbsolutePathForExternalAppForWrite(*FPaths::GetProjectFilePath());
-			const FString Arguments = FString::Printf(TEXT("%s %s %s -Plugin=\"%s\" -Project=\"%s\" -Progress -NoHotReloadFromIDE"), FPlatformMisc::GetUBTTargetName(), FModuleManager::Get().GetUBTConfiguration(), FPlatformMisc::GetUBTPlatform(), *PluginFilePath, *ProjectFileName);
+
+			// Internal plugins will be found and built automatically, and using the -Plugin param will mark them as external plugins causing them to fail to load.
+			const bool bIsEnginePlugin = FPaths::IsUnderDirectory(PluginFilePath, FPaths::EnginePluginsDir());
+			const bool bIsProjectPlugin = FPaths::IsUnderDirectory(PluginFilePath, FPaths::ProjectPluginsDir());
+			const FString ExternalPluginArgument = !(bIsEnginePlugin || bIsProjectPlugin) ? FString::Printf(TEXT("-Plugin=\"%s\" "), *PluginFilePath) : FString();
+
+			const FString Arguments = FString::Printf(TEXT("%s %s %s %s-Project=\"%s\" -Progress -NoHotReloadFromIDE"), FPlatformMisc::GetUBTTargetName(), FModuleManager::Get().GetUBTConfiguration(), FPlatformMisc::GetUBTPlatform(), *ExternalPluginArgument, *ProjectFileName);
+
 			if (!FDesktopPlatformModule::Get()->RunUnrealBuildTool(FText::Format(LOCTEXT("CompilingPlugin", "Compiling {0} plugin..."), FText::FromString(PluginName)), FPaths::RootDir(), Arguments, GWarn))
 			{
 				if (LoadParams.OutFailReason)
