Adding entries for CanContainVerse and VersePath to FNewPluginParams

diff --git a/Engine/Plugins/Developer/PluginUtils/Source/PluginUtils/Private/PluginUtils.cpp b/Engine/Plugins/Developer/PluginUtils/Source/PluginUtils/Private/PluginUtils.cpp
--- a/Engine/Plugins/Developer/PluginUtils/Source/PluginUtils/Private/PluginUtils.cpp
+++ b/Engine/Plugins/Developer/PluginUtils/Source/PluginUtils/Private/PluginUtils.cpp
@@ -489,8 +489,10 @@ TSharedPtr<IPlugin> FPluginUtils::CreateAndLoadNewPlugin(const FString& PluginNa
 	ExCreationParams.Descriptor.Description = CreationParams.Description;
 	ExCreationParams.Descriptor.bIsBetaVersion = CreationParams.bIsBetaVersion;
 	ExCreationParams.Descriptor.bCanContainContent = CreationParams.bCanContainContent;
+	ExCreationParams.Descriptor.bCanContainVerse = CreationParams.bCanContainVerse;
 	ExCreationParams.Descriptor.EnabledByDefault = CreationParams.EnabledByDefault;
 	ExCreationParams.Descriptor.bExplicitlyLoaded = CreationParams.bExplicitelyLoaded;
+	ExCreationParams.Descriptor.VersePath = CreationParams.VersePath;
 
 	if (CreationParams.bHasModules)
 	{
diff --git a/Engine/Plugins/Developer/PluginUtils/Source/PluginUtils/Public/PluginUtils.h b/Engine/Plugins/Developer/PluginUtils/Source/PluginUtils/Public/PluginUtils.h
--- a/Engine/Plugins/Developer/PluginUtils/Source/PluginUtils/Public/PluginUtils.h
+++ b/Engine/Plugins/Developer/PluginUtils/Source/PluginUtils/Public/PluginUtils.h
@@ -77,6 +77,9 @@ public:
 		/** Can this plugin contain content */
 		bool bCanContainContent = false;
 
+		/** Can this plugin contain Verse */
+		bool bCanContainVerse = false;
+
 		/** Does this plugin have Source files? */
 		bool bHasModules = false;
 
@@ -86,6 +89,9 @@ public:
 		 */
 		bool bExplicitelyLoaded = false;
 
+		/** The Verse path to the root of this plugin's content directory */
+		FString VersePath;
+
 		/** Whether this plugin should be enabled/disabled by default for any project. */
 		EPluginEnabledByDefault EnabledByDefault = EPluginEnabledByDefault::Unspecified;
 
