Add global name filter across editor for asset paths

diff --git a/Engine/Plugins/Developer/PluginUtils/Source/PluginUtils/Private/PluginUtils.cpp b/Engine/Plugins/Developer/PluginUtils/Source/PluginUtils/Private/PluginUtils.cpp
--- a/Engine/Plugins/Developer/PluginUtils/Source/PluginUtils/Private/PluginUtils.cpp
+++ b/Engine/Plugins/Developer/PluginUtils/Source/PluginUtils/Private/PluginUtils.cpp
@@ -1090,55 +1090,56 @@ bool FPluginUtils::IsValidPluginName(const FString& PluginName, FText* FailReaso
 	static const FText PluginTerm = LOCTEXT("PluginTerm", "Plugin");
 	const FText& PluginTermToUse = PluginTermReplacement ? *PluginTermReplacement : PluginTerm;
 
-	bool bIsNameValid = true;
-
 	// Cannot be empty
 	if (PluginName.IsEmpty())
 	{
-		bIsNameValid = false;
 		if (FailReason)
 		{
 			*FailReason = FText::Format(LOCTEXT("PluginNameIsEmpty", "{0} name cannot be empty"), PluginTermToUse);
 		}
+		return false;
 	}
 
 	// Must begin with an alphabetic character
-	if (bIsNameValid && !FChar::IsAlpha(PluginName[0]))
+	if (!FChar::IsAlpha(PluginName[0]))
 	{
-		bIsNameValid = false;
 		if (FailReason)
 		{
 			*FailReason = FText::Format(LOCTEXT("PluginNameMustBeginWithAlphabetic", "{0} name must begin with an alphabetic character"), PluginTermToUse);
 		}
+		return false;
 	}
 
 	// Only allow alphanumeric characters and underscore in the name
-	if (bIsNameValid)
+	FString IllegalCharacters;
+	for (int32 CharIdx = 0; CharIdx < PluginName.Len(); ++CharIdx)
 	{
-		FString IllegalCharacters;
-		for (int32 CharIdx = 0; CharIdx < PluginName.Len(); ++CharIdx)
+		const FString& Char = PluginName.Mid(CharIdx, 1);
+		if (!FChar::IsAlnum(Char[0]) && Char != TEXT("_") && Char != TEXT("-"))
 		{
-			const FString& Char = PluginName.Mid(CharIdx, 1);
-			if (!FChar::IsAlnum(Char[0]) && Char != TEXT("_") && Char != TEXT("-"))
+			if (!IllegalCharacters.Contains(Char))
 			{
-				if (!IllegalCharacters.Contains(Char))
-				{
-					IllegalCharacters += Char;
-				}
+				IllegalCharacters += Char;
 			}
 		}
-		
-		if (IllegalCharacters.Len() > 0)
+	}
+
+	if (IllegalCharacters.Len() > 0)
+	{
+		if (FailReason)
 		{
-			bIsNameValid = false;
-			if (FailReason)
-			{
-				*FailReason = FText::Format(LOCTEXT("PluginNameContainsIllegalCharacters", "{0} name cannot contain characters such as \"{1}\""), PluginTermToUse, FText::FromString(IllegalCharacters));
-			}
+			*FailReason = FText::Format(LOCTEXT("PluginNameContainsIllegalCharacters", "{0} name cannot contain characters such as \"{1}\""), PluginTermToUse, FText::FromString(IllegalCharacters));
 		}
+		return false;
 	}
 
-	return bIsNameValid;
+	FAssetToolsModule& AssetToolsModule = FModuleManager::LoadModuleChecked<FAssetToolsModule>("AssetTools");
+	if (!AssetToolsModule.Get().IsNameAllowed(PluginName, FailReason))
+	{
+		return false;
+	}
+
+	return true;
 }
 
 #undef LOCTEXT_NAMESPACE
diff --git a/Engine/Plugins/Editor/ContentBrowser/ContentBrowserFileDataSource/Source/ContentBrowserFileDataSource/ContentBrowserFileDataSource.Build.cs b/Engine/Plugins/Editor/ContentBrowser/ContentBrowserFileDataSource/Source/ContentBrowserFileDataSource/ContentBrowserFileDataSource.Build.cs
--- a/Engine/Plugins/Editor/ContentBrowser/ContentBrowserFileDataSource/Source/ContentBrowserFileDataSource/ContentBrowserFileDataSource.Build.cs
+++ b/Engine/Plugins/Editor/ContentBrowser/ContentBrowserFileDataSource/Source/ContentBrowserFileDataSource/ContentBrowserFileDataSource.Build.cs
@@ -20,6 +20,7 @@ namespace UnrealBuildTool.Rules
 				new string[]
 				{
 					"EditorFramework",
+					"AssetTools",
 					"UnrealEd",
 					"ToolMenus",
 					"Slate",
diff --git a/Engine/Plugins/Editor/ContentBrowser/ContentBrowserFileDataSource/Source/ContentBrowserFileDataSource/Private/ContentBrowserFileDataCore.cpp b/Engine/Plugins/Editor/ContentBrowser/ContentBrowserFileDataSource/Source/ContentBrowserFileDataSource/Private/ContentBrowserFileDataCore.cpp
--- a/Engine/Plugins/Editor/ContentBrowser/ContentBrowserFileDataSource/Source/ContentBrowserFileDataSource/Private/ContentBrowserFileDataCore.cpp
+++ b/Engine/Plugins/Editor/ContentBrowser/ContentBrowserFileDataSource/Source/ContentBrowserFileDataSource/Private/ContentBrowserFileDataCore.cpp
@@ -8,6 +8,7 @@
 #include "Misc/Paths.h"
 #include "AssetRegistry/AssetData.h"
 #include "AssetThumbnail.h"
+#include "AssetToolsModule.h"
 
 #define LOCTEXT_NAMESPACE "ContentBrowserFileDataSource"
 
@@ -639,6 +640,13 @@ bool IsValidName(const bool bIsDirectory, const bool bCheckUniqueName, const FSt
 		return false;
 	}
 
+	// Check custom filter set by external module
+	FAssetToolsModule& AssetToolsModule = FModuleManager::LoadModuleChecked<FAssetToolsModule>("AssetTools");
+	if (!AssetToolsModule.Get().IsNameAllowed(InNewName, OutErrorMsg))
+	{
+		return false;
+	}
+
 	// Is this name unique?
 	if (bCheckUniqueName)
 	{
diff --git a/Engine/Source/Developer/AssetTools/Private/AssetTools.cpp b/Engine/Source/Developer/AssetTools/Private/AssetTools.cpp
--- a/Engine/Source/Developer/AssetTools/Private/AssetTools.cpp
+++ b/Engine/Source/Developer/AssetTools/Private/AssetTools.cpp
@@ -4219,4 +4219,26 @@ void UAssetToolsImpl::NotifyBlockedByWritableFolderFilter() const
 	FSlateNotificationManager::Get().AddNotification(FNotificationInfo(LOCTEXT("NotifyBlockedByWritableFolderFilter", "Folder is locked")));
 }
 
+bool UAssetToolsImpl::IsNameAllowed(const FString& Name, FText* OutErrorMessage) const
+{
+	for (const TPair<FName, FIsNameAllowed>& DelegatePair : IsNameAllowedDelegates)
+	{
+		if (!DelegatePair.Value.Execute(Name, OutErrorMessage))
+		{
+			return false;
+		}
+	}
+	return true;
+}
+
+void UAssetToolsImpl::RegisterIsNameAllowedDelegate(const FName OwnerName, FIsNameAllowed Delegate)
+{
+	IsNameAllowedDelegates.Add(OwnerName, Delegate);
+}
+
+void UAssetToolsImpl::UnregisterIsNameAllowedDelegate(const FName OwnerName)
+{
+	IsNameAllowedDelegates.Remove(OwnerName);
+}
+
 #undef LOCTEXT_NAMESPACE
diff --git a/Engine/Source/Developer/AssetTools/Private/AssetTools.h b/Engine/Source/Developer/AssetTools/Private/AssetTools.h
--- a/Engine/Source/Developer/AssetTools/Private/AssetTools.h
+++ b/Engine/Source/Developer/AssetTools/Private/AssetTools.h
@@ -132,6 +132,9 @@ public:
 	virtual TSharedRef<FPathPermissionList>& GetWritableFolderPermissionList() override;
 	virtual bool AllPassWritableFolderFilter(const TArray<FString>& InPaths) const override;
 	virtual void NotifyBlockedByWritableFolderFilter() const;
+	virtual bool IsNameAllowed(const FString& Name, FText* OutErrorMessage = nullptr) const override;
+	virtual void RegisterIsNameAllowedDelegate(const FName OwnerName, FIsNameAllowed Delegate) override;
+	virtual void UnregisterIsNameAllowedDelegate(const FName OwnerName) override;
 	
 	virtual void SyncBrowserToAssets(const TArray<UObject*>& AssetsToSync) override;
 	virtual void SyncBrowserToAssets(const TArray<FAssetData>& AssetsToSync) override;
@@ -231,6 +234,8 @@ private:
 	TArray<FString> SubContentDenyListPaths;
 
 	bool CreateAssetsAsExternallyReferenceable;
+
+	TMap<FName, FIsNameAllowed> IsNameAllowedDelegates;
 };
 
 PRAGMA_ENABLE_DEPRECATION_WARNINGS
diff --git a/Engine/Source/Developer/AssetTools/Private/AssetViewUtils.cpp b/Engine/Source/Developer/AssetTools/Private/AssetViewUtils.cpp
--- a/Engine/Source/Developer/AssetTools/Private/AssetViewUtils.cpp
+++ b/Engine/Source/Developer/AssetTools/Private/AssetViewUtils.cpp
@@ -1041,6 +1041,13 @@ bool AssetViewUtils::IsValidFolderName(const FString& FolderName, FText& Reason)
 			return false;
 		}
 	}
+
+	// Check custom filter set by external module
+	FAssetToolsModule& AssetToolsModule = FModuleManager::LoadModuleChecked<FAssetToolsModule>("AssetTools");
+	if (!AssetToolsModule.Get().IsNameAllowed(FolderName, &Reason))
+	{
+		return false;
+	}
 	
 	return FFileHelper::IsFilenameValidForSaving( FolderName, Reason );
 }
@@ -1314,6 +1321,13 @@ bool AssetViewUtils::IsValidObjectPathForCreate(const FString& ObjectPath, const
 		}
 	}
 
+	// Check custom filter set by external module
+	FAssetToolsModule& AssetToolsModule = FModuleManager::LoadModuleChecked<FAssetToolsModule>("AssetTools");
+	if (!AssetToolsModule.Get().IsNameAllowed(ObjectName, &OutErrorMessage))
+	{
+		return false;
+	}
+
 	return true;
 }
 
diff --git a/Engine/Source/Developer/AssetTools/Public/IAssetTools.h b/Engine/Source/Developer/AssetTools/Public/IAssetTools.h
--- a/Engine/Source/Developer/AssetTools/Public/IAssetTools.h
+++ b/Engine/Source/Developer/AssetTools/Public/IAssetTools.h
@@ -108,6 +108,7 @@ struct FAssetRenameData
 
 
 DECLARE_MULTICAST_DELEGATE_OneParam(FAssetPostRenameEvent, const TArray<FAssetRenameData>&);
+DECLARE_DELEGATE_RetVal_TwoParams(bool, FIsNameAllowed, const FString& /*Name*/, FText* /*OutErrorMessage*/);
 
 
 struct FAdvancedAssetCategory
@@ -486,6 +487,13 @@ public:
 	/** Returns true if all in list pass writable folder filter */
 	virtual bool AllPassWritableFolderFilter(const TArray<FString>& InPaths) const = 0;
 
+	/** Returns true if IsNameAllowedDelegate is not set, or if the name passes the filter function*/
+	virtual bool IsNameAllowed(const FString& Name, FText* OutErrorMessage = nullptr) const = 0;
+	/** Allows setting of a global name filter that is applied to folders, assets, plugins, etc. */
+	virtual void RegisterIsNameAllowedDelegate(const FName OwnerName, FIsNameAllowed Delegate) = 0;
+	/** Remove a previously-set global name filter */
+	virtual void UnregisterIsNameAllowedDelegate(const FName OwnerName) = 0;
+
 	/** Show notification that writable folder filter blocked an action */
 	virtual void NotifyBlockedByWritableFolderFilter() const = 0;
 
