FPluginUtils: Close the current editor map if it's in one of the plugins to unload (maps and external actors can't be properly unloaded via UnloadPackages)

diff --git a/Engine/Plugins/Developer/PluginUtils/Source/PluginUtils/Private/PluginUtils.cpp b/Engine/Plugins/Developer/PluginUtils/Source/PluginUtils/Private/PluginUtils.cpp
--- a/Engine/Plugins/Developer/PluginUtils/Source/PluginUtils/Private/PluginUtils.cpp
+++ b/Engine/Plugins/Developer/PluginUtils/Source/PluginUtils/Private/PluginUtils.cpp
@@ -6,6 +6,8 @@
 #include "SourceControlHelpers.h"
 #include "IContentBrowserSingleton.h"
 #include "ContentBrowserModule.h"
+#include "Editor.h"
+#include "FileHelpers.h"
 #include "GameProjectUtils.h"
 #include "Interfaces/IProjectManager.h"
 #include "Interfaces/IPluginManager.h"
@@ -20,6 +22,7 @@
 #include "HAL/FileManager.h"
 #include "HAL/PlatformFileManager.h"
 #include "Misc/Paths.h"
+#include "Misc/PathViews.h"
 #include "Misc/FileHelper.h"
 #include "Misc/PackageName.h"
 #include "Misc/FeedbackContext.h"
@@ -869,13 +872,15 @@ bool FPluginUtils::UnloadPluginAssets(const FString& PluginName, FText* OutFailR
 
 bool FPluginUtils::UnloadPluginsAssets(const TConstArrayView<TSharedRef<IPlugin>> Plugins, FText* OutFailReason /*= nullptr*/)
 {
-	if (Plugins.Num() == 0)
+	if (Plugins.IsEmpty())
 	{
 		return true;
 	}
 
 	TArray<FString> PluginContentMountPoints;
+	TSet<FString> PluginNames;
 	PluginContentMountPoints.Reserve(Plugins.Num());
+	PluginNames.Reserve(Plugins.Num());
 
 	for (const TSharedRef<IPlugin>& Plugin : Plugins)
 	{
@@ -885,18 +890,35 @@ bool FPluginUtils::UnloadPluginsAssets(const TConstArrayView<TSharedRef<IPlugin>
 			if (FPackageName::MountPointExists(PluginContentMountPoint))
 			{
 				PluginContentMountPoints.Add(MoveTemp(PluginContentMountPoint));
+				PluginNames.Add(Plugin->GetName());
 			}
 		}
 	}
 
-	// Synchronous scan plugins to make sure we find all their assets.
-	IAssetRegistry& AssetRegistry = FModuleManager::LoadModuleChecked<FAssetRegistryModule>(AssetRegistryConstants::ModuleName).Get();
+	if (PluginContentMountPoints.IsEmpty())
+	{
+		return true;
+	}
 
-	if (!PluginContentMountPoints.IsEmpty())
+	// Close the current map if it's in one of the plugins
+	if (UWorld* CurrentWorld = GEditor ? GEditor->GetEditorWorldContext().World() : nullptr)
 	{
-		AssetRegistry.ScanPathsSynchronous(PluginContentMountPoints, /*bForceRescan=*/ true);
+		const FNameBuilder CurrentWorldPath(CurrentWorld->GetPackage()->GetFName());
+		const FStringView CurrentWorldMountPoint = FPathViews::GetMountPointNameFromPath(CurrentWorldPath);
+		if (PluginNames.ContainsByHash(GetTypeHash(CurrentWorldMountPoint), CurrentWorldMountPoint))
+		{
+			UWorld* BlankWorld = UEditorLoadingAndSavingUtils::NewBlankMap(/*bSaveExistingMap*/ false);
+			if (ensure(BlankWorld))
+			{
+				BlankWorld->GetOutermost()->ClearDirtyFlag();
+			}
+		}
 	}
 
+	// Synchronous scan plugins to make sure we find all their assets.
+	IAssetRegistry& AssetRegistry = FModuleManager::LoadModuleChecked<FAssetRegistryModule>(AssetRegistryConstants::ModuleName).Get();
+	AssetRegistry.ScanPathsSynchronous(PluginContentMountPoints, /*bForceRescan=*/ true);
+
 	// Unload plugin packages
 	{
 		FARFilter ARFilter;
