Fixing the plugin manager not resetting asset localization ids on creation.

diff --git a/Engine/Plugins/Developer/PluginUtils/Source/PluginUtils/Private/PluginUtils.cpp b/Engine/Plugins/Developer/PluginUtils/Source/PluginUtils/Private/PluginUtils.cpp
--- a/Engine/Plugins/Developer/PluginUtils/Source/PluginUtils/Private/PluginUtils.cpp
+++ b/Engine/Plugins/Developer/PluginUtils/Source/PluginUtils/Private/PluginUtils.cpp
@@ -20,6 +20,7 @@
 #include "PackageTools.h"
 #include "HAL/FileManager.h"
 #include "HAL/PlatformFileManager.h"
+#include "Internationalization/TextPackageNamespaceUtil.h"
 #include "Misc/Paths.h"
 #include "Misc/PathViews.h"
 #include "Misc/FileHelper.h"
@@ -257,7 +258,7 @@ namespace PluginUtils
 
 			/**
 			 * Fixes up any assets that contain the PLUGIN_NAME text macro, since those need to be renamed by the engine for the change to
-			 * stick (as opposed to just renaming the file)
+			 * stick (as opposed to just renaming the file) as well as regenerating the localization ids for the assets in the plugin.
 			 */
 			void PerformFixup(TMap<FString, FString>& OutModifiedAssetPaths)
 			{
@@ -279,14 +280,21 @@ namespace PluginUtils
 							AssetRegistry.GetAssetsByPackageName(*PackageName, Assets);
 						}
 
-						for (FAssetData Asset : Assets)
+						for (FAssetData AssetData : Assets)
 						{
-							const FString AssetName = Asset.AssetName.ToString().Replace(*PLUGIN_NAME, *PluginName, ESearchCase::CaseSensitive);
-							const FString AssetPath = Asset.PackagePath.ToString().Replace(*PLUGIN_NAME, *PluginName, ESearchCase::CaseSensitive);
-							AssetToOriginalFilePathMap.Add(Asset.GetAsset(), File);
-							FAssetRenameData RenameData(Asset.GetAsset(), AssetPath, AssetName);
+							const FString AssetName = AssetData.AssetName.ToString().Replace(*PLUGIN_NAME, *PluginName, ESearchCase::CaseSensitive);
+							const FString AssetPath = AssetData.PackagePath.ToString().Replace(*PLUGIN_NAME, *PluginName, ESearchCase::CaseSensitive);
+							AssetToOriginalFilePathMap.Add(AssetData.GetAsset(), File);
+							FAssetRenameData RenameData(AssetData.GetAsset(), AssetPath, AssetName);
 
 							AssetRenameData.Add(RenameData);
+
+							if (UObject* Asset = AssetData.GetAsset())
+							{
+								Asset->Modify();
+								TextNamespaceUtil::ClearPackageNamespace(Asset);
+								TextNamespaceUtil::EnsurePackageNamespace(Asset);
+							}
 						}
 					}
 
