Re-enabling LocalContentWorker test now that core issues with AssetDataGatherer sync load have been fixed via CL: 23521344
Fix an issue where plugin unloading would only find packages that were in the asset registry. This prevented the cleanup of some _externalActors_ packages.
Add nullptr checks on StreamingPolicy which might be null for a transient world
fix an issue where a plugin might not be fully removed during the duplicate process

diff --git a/Engine/Plugins/Developer/PluginUtils/Source/PluginUtils/Private/PluginUtils.cpp b/Engine/Plugins/Developer/PluginUtils/Source/PluginUtils/Private/PluginUtils.cpp
--- a/Engine/Plugins/Developer/PluginUtils/Source/PluginUtils/Private/PluginUtils.cpp
+++ b/Engine/Plugins/Developer/PluginUtils/Source/PluginUtils/Private/PluginUtils.cpp
@@ -911,48 +911,45 @@ bool FPluginUtils::UnloadPluginsAssets(const TConstArrayView<TSharedRef<IPlugin>
 	// Synchronous scan plugins to make sure we find all their assets.
 	IAssetRegistry& AssetRegistry = FModuleManager::LoadModuleChecked<FAssetRegistryModule>(AssetRegistryConstants::ModuleName).Get();
 	AssetRegistry.ScanPathsSynchronous(PluginContentMountPoints, /*bForceRescan=*/ true);
-
-	// Unload plugin packages
+	
+	TArray<FName> AllPluginPackages;
+	auto FindAllPluginPackages = [&AllPluginPackages, PluginContentMountPoints](FName PackageName, const FAssetPackageData& PackageData)
 	{
-		FARFilter ARFilter;
-		ARFilter.PackagePaths.Reserve(PluginContentMountPoints.Num());
-		for (const FString& PluginContentMountPoint : PluginContentMountPoints)
+		const FString PackageNameStr = PackageName.ToString();
+		for (const FString& PluginMountPoint : PluginContentMountPoints)
 		{
-			FString PluginRoot = PluginContentMountPoint;
-			PluginRoot.RemoveFromEnd(TEXT("/"));
-			ARFilter.PackagePaths.Add(*PluginRoot);
+			if (PackageNameStr.StartsWith(PluginMountPoint, ESearchCase::IgnoreCase))
+			{
+				AllPluginPackages.Add(PackageName);
+			}
 		}
-		ARFilter.bRecursivePaths = true;
+	};
+	AssetRegistry.EnumerateAllPackages(FindAllPluginPackages);
 
-		TArray<FAssetData> PluginAssets;
-		if (AssetRegistry.GetAssets(ARFilter, PluginAssets))
+	TSet<UPackage*> PackagesToUnload;
+	PackagesToUnload.Reserve(AllPluginPackages.Num());
+	for (const FName& PackageName : AllPluginPackages)
+	{
+		if (UPackage* Package = FindPackage(NULL, *PackageName.ToString()))
 		{
-			TSet<UPackage*> PackagesToUnload;
-			PackagesToUnload.Reserve(PluginAssets.Num());
-			for (const FAssetData& AssetData : PluginAssets)
-			{
-				if (UPackage* Package = FindPackage(NULL, *AssetData.PackageName.ToString()))
-				{
-					PackagesToUnload.Add(Package);
-				}
-			}
+			PackagesToUnload.Add(Package);
+		}
+	}
 
-			if (PackagesToUnload.Num() > 0)
-			{
-				FText ErrorMsg;
-				UPackageTools::UnloadPackages(PackagesToUnload.Array(), ErrorMsg, /*bUnloadDirtyPackages=*/true);
+	if (PackagesToUnload.Num() > 0)
+	{
+		FText ErrorMsg;
+		UPackageTools::UnloadPackages(PackagesToUnload.Array(), ErrorMsg, /*bUnloadDirtyPackages=*/true);
 
-				// @note UnloadPackages returned bool indicates whether some packages were unloaded
-				// To tell whether all packages were successfully unloaded we must check the ErrorMsg output param
-				if (!ErrorMsg.IsEmpty())
-				{
-					if (OutFailReason)
-					{
-						*OutFailReason = MoveTemp(ErrorMsg);
-					}
-					return false;
-				}
+		// @note UnloadPackages returned bool indicates whether some packages were unloaded
+		// To tell whether all packages were successfully unloaded we must check the ErrorMsg output param
+		if (!ErrorMsg.IsEmpty())
+		{
+			if (OutFailReason)
+			{
+				*OutFailReason = MoveTemp(ErrorMsg);
 			}
+			return false;
 		}
 	}
 
diff --git a/Engine/Source/Runtime/Engine/Private/WorldPartition/WorldPartition.cpp b/Engine/Source/Runtime/Engine/Private/WorldPartition/WorldPartition.cpp
--- a/Engine/Source/Runtime/Engine/Private/WorldPartition/WorldPartition.cpp
+++ b/Engine/Source/Runtime/Engine/Private/WorldPartition/WorldPartition.cpp
@@ -1438,7 +1438,10 @@ void UWorldPartition::UpdateStreamingState()
 
 	if (GetWorld()->IsGameWorld())
 	{
-		StreamingPolicy->UpdateStreamingState();
+		if (StreamingPolicy)
+		{
+			StreamingPolicy->UpdateStreamingState();
+		}
 	}
 }
 
@@ -1893,7 +1896,10 @@ void UWorldPartition::OnWorldRenamed(UWorld* RenamedWorld)
 
 void UWorldPartition::RemapSoftObjectPath(FSoftObjectPath& ObjectPath)
 {
-	StreamingPolicy->RemapSoftObjectPath(ObjectPath);
+	if (StreamingPolicy)
+	{
+		StreamingPolicy->RemapSoftObjectPath(ObjectPath);
+	}
 }
 
 FBox UWorldPartition::GetEditorWorldBounds() const
