Optimize UnloadPluginsAssets

diff --git a/Engine/Plugins/Developer/PluginUtils/Source/PluginUtils/Private/PluginUtils.cpp b/Engine/Plugins/Developer/PluginUtils/Source/PluginUtils/Private/PluginUtils.cpp
--- a/Engine/Plugins/Developer/PluginUtils/Source/PluginUtils/Private/PluginUtils.cpp
+++ b/Engine/Plugins/Developer/PluginUtils/Source/PluginUtils/Private/PluginUtils.cpp
@@ -26,6 +26,7 @@
 #include "Misc/FileHelper.h"
 #include "Misc/PackageName.h"
 #include "Misc/FeedbackContext.h"
+#include "UObject/UObjectIterator.h"
 
 DEFINE_LOG_CATEGORY_STATIC(LogPluginUtils, Log, All);
 
@@ -866,80 +867,46 @@ bool FPluginUtils::UnloadPlugins(const TConstArrayView<FString> PluginNames, FTe
 
 bool FPluginUtils::UnloadPluginAssets(const TSharedRef<IPlugin>& Plugin, FText* OutFailReason /*= nullptr*/)
 {
-	return UnloadPluginsAssets({ Plugin }, OutFailReason);
+	return UnloadPluginAssets(Plugin->GetName(), OutFailReason);
 }
 
 bool FPluginUtils::UnloadPluginAssets(const FString& PluginName, FText* OutFailReason /*= nullptr*/)
 {
-	if (TSharedPtr<IPlugin> Plugin = IPluginManager::Get().FindPlugin(PluginName))
-	{
-		return UnloadPluginAssets(Plugin.ToSharedRef(), OutFailReason);
-	}
-	return true;
+	TSet<FString> PluginNames;
+	PluginNames.Reserve(1);
+	PluginNames.Add(PluginName);
+	return UnloadPluginsAssets(PluginNames, OutFailReason);
 }
 
 bool FPluginUtils::UnloadPluginsAssets(const TConstArrayView<TSharedRef<IPlugin>> Plugins, FText* OutFailReason /*= nullptr*/)
 {
-	if (Plugins.IsEmpty())
-	{
-		return true;
-	}
-
-	TArray<FString> PluginContentMountPoints;
 	TSet<FString> PluginNames;
-	PluginContentMountPoints.Reserve(Plugins.Num());
 	PluginNames.Reserve(Plugins.Num());
-
 	for (const TSharedRef<IPlugin>& Plugin : Plugins)
 	{
-		if (Plugin->IsEnabled())
-		{
-			FString PluginContentMountPoint = Plugin->GetMountedAssetPath();
-			if (FPackageName::MountPointExists(PluginContentMountPoint))
-			{
-				PluginContentMountPoints.Add(MoveTemp(PluginContentMountPoint));
-				PluginNames.Add(Plugin->GetName());
-			}
-		}
-	}
-
-	if (PluginContentMountPoints.IsEmpty())
-	{
-		return true;
+		PluginNames.Add(Plugin->GetName());
 	}
 
-	// Synchronous scan plugins to make sure we find all their assets.
-	IAssetRegistry& AssetRegistry = FModuleManager::LoadModuleChecked<FAssetRegistryModule>(AssetRegistryConstants::ModuleName).Get();
-	AssetRegistry.ScanPathsSynchronous(PluginContentMountPoints, /*bForceRescan=*/ true);
-	
-	TArray<FName> AllPluginPackages;
-	auto FindAllPluginPackages = [&AllPluginPackages, PluginContentMountPoints](FName PackageName, const FAssetPackageData& PackageData)
-	{
-		const FString PackageNameStr = PackageName.ToString();
-		for (const FString& PluginMountPoint : PluginContentMountPoints)
-		{
-			if (PackageNameStr.StartsWith(PluginMountPoint, ESearchCase::IgnoreCase))
-			{
-				AllPluginPackages.Add(PackageName);
-			}
-		}
-	};
-	AssetRegistry.EnumerateAllPackages(FindAllPluginPackages);
+	return UnloadPluginsAssets(PluginNames, OutFailReason);
+}
 
-	TSet<UPackage*> PackagesToUnload;
-	PackagesToUnload.Reserve(AllPluginPackages.Num());
-	for (const FName& PackageName : AllPluginPackages)
+bool FPluginUtils::UnloadPluginsAssets(const TSet<FString>& PluginNames, FText* OutFailReason /*= nullptr*/)
+{
+	TArray<UPackage*> PackagesToUnload;
+	for (TObjectIterator<UPackage> It; It; ++It)
 	{
-		if (UPackage* Package = FindPackage(NULL, *PackageName.ToString()))
+		const FNameBuilder PackageName(It->GetFName());
+		const FStringView PackageMountPointName = FPathViews::GetMountPointNameFromPath(PackageName);
+		if (PluginNames.ContainsByHash(GetTypeHash(PackageMountPointName), PackageMountPointName))
 		{
-			PackagesToUnload.Add(Package);
+			PackagesToUnload.Add(*It);
 		}
 	}
 
 	if (PackagesToUnload.Num() > 0)
 	{
 		FText ErrorMsg;
-		UPackageTools::UnloadPackages(PackagesToUnload.Array(), ErrorMsg, /*bUnloadDirtyPackages=*/true);
+		UPackageTools::UnloadPackages(PackagesToUnload, ErrorMsg, /*bUnloadDirtyPackages=*/true);
 
 		// @note UnloadPackages returned bool indicates whether some packages were unloaded
 		// To tell whether all packages were successfully unloaded we must check the ErrorMsg output param
@@ -958,16 +925,13 @@ bool FPluginUtils::UnloadPluginsAssets(const TConstArrayView<TSharedRef<IPlugin>
 
 bool FPluginUtils::UnloadPluginsAssets(const TConstArrayView<FString> PluginNames, FText* OutFailReason /*= nullptr*/)
 {
-	TArray<TSharedRef<IPlugin>> Plugins;
-	Plugins.Reserve(PluginNames.Num());
+	TSet<FString> PluginNamesSet;
+	PluginNamesSet.Reserve(PluginNames.Num());
 	for (const FString& PluginName : PluginNames)
 	{
-		if (TSharedPtr<IPlugin> Plugin = IPluginManager::Get().FindPlugin(PluginName))
-		{
-			Plugins.Add(Plugin.ToSharedRef());
-		}
+		PluginNamesSet.Add(PluginName);
 	}
-	return UnloadPluginsAssets(Plugins, OutFailReason);
+	return UnloadPluginsAssets(PluginNamesSet, OutFailReason);
 }
 
 bool FPluginUtils::AddToPluginSearchPathIfNeeded(const FString& Dir, bool bRefreshPlugins, bool bUpdateProjectFile)
diff --git a/Engine/Plugins/Developer/PluginUtils/Source/PluginUtils/Public/PluginUtils.h b/Engine/Plugins/Developer/PluginUtils/Source/PluginUtils/Public/PluginUtils.h
--- a/Engine/Plugins/Developer/PluginUtils/Source/PluginUtils/Public/PluginUtils.h
+++ b/Engine/Plugins/Developer/PluginUtils/Source/PluginUtils/Public/PluginUtils.h
@@ -348,6 +348,15 @@ public:
 	 */
 	static bool UnloadPluginsAssets(const TConstArrayView<TSharedRef<IPlugin>> Plugins, FText* OutFailReason = nullptr);
 
+	/**
+	 * Unload assets from the specified plugin but does not unmount them
+	 * @warning Dirty assets that need to be saved will be unloaded anyway
+	 * @param PluginNames Names of the plugins to unload assets from
+	 * @param OutFailReason Outputs the reason of the failure if any
+	 * @return Whether plugin assets were successfully unloaded
+	 */
+	static bool UnloadPluginsAssets(const TSet<FString>& PluginNames, FText* OutFailReason = nullptr);
+
 	/**
 	 * Unload assets from the specified plugin but does not unmount them
 	 * @warning Dirty assets that need to be saved will be unloaded anyway
