Log time to UnloadPluginsAssets

diff --git a/Engine/Plugins/Developer/PluginUtils/Source/PluginUtils/Private/PluginUtils.cpp b/Engine/Plugins/Developer/PluginUtils/Source/PluginUtils/Private/PluginUtils.cpp
--- a/Engine/Plugins/Developer/PluginUtils/Source/PluginUtils/Private/PluginUtils.cpp
+++ b/Engine/Plugins/Developer/PluginUtils/Source/PluginUtils/Private/PluginUtils.cpp
@@ -890,35 +890,42 @@ bool FPluginUtils::UnloadPluginsAssets(const TConstArrayView<TSharedRef<IPlugin>
 
 bool FPluginUtils::UnloadPluginsAssets(const TSet<FString>& PluginNames, FText* OutFailReason /*= nullptr*/)
 {
-	TArray<UPackage*> PackagesToUnload;
-	for (TObjectIterator<UPackage> It; It; ++It)
+	bool bSuccess = true;
+	if (!PluginNames.IsEmpty())
 	{
-		const FNameBuilder PackageName(It->GetFName());
-		const FStringView PackageMountPointName = FPathViews::GetMountPointNameFromPath(PackageName);
-		if (PluginNames.ContainsByHash(GetTypeHash(PackageMountPointName), PackageMountPointName))
+		const double StartTime = FPlatformTime::Seconds();
+
+		TArray<UPackage*> PackagesToUnload;
+		for (TObjectIterator<UPackage> It; It; ++It)
 		{
-			PackagesToUnload.Add(*It);
+			const FNameBuilder PackageName(It->GetFName());
+			const FStringView PackageMountPointName = FPathViews::GetMountPointNameFromPath(PackageName);
+			if (PluginNames.ContainsByHash(GetTypeHash(PackageMountPointName), PackageMountPointName))
+			{
+				PackagesToUnload.Add(*It);
+			}
 		}
-	}
 
-	if (PackagesToUnload.Num() > 0)
-	{
-		FText ErrorMsg;
-		UPackageTools::UnloadPackages(PackagesToUnload, ErrorMsg, /*bUnloadDirtyPackages=*/true);
-
-		// @note UnloadPackages returned bool indicates whether some packages were unloaded
-		// To tell whether all packages were successfully unloaded we must check the ErrorMsg output param
-		if (!ErrorMsg.IsEmpty())
+		if (PackagesToUnload.Num() > 0)
 		{
-			if (OutFailReason)
+			FText ErrorMsg;
+			UPackageTools::UnloadPackages(PackagesToUnload, ErrorMsg, /*bUnloadDirtyPackages=*/true);
+
+			// @note UnloadPackages returned bool indicates whether some packages were unloaded
+			// To tell whether all packages were successfully unloaded we must check the ErrorMsg output param
+			if (!ErrorMsg.IsEmpty())
 			{
-				*OutFailReason = MoveTemp(ErrorMsg);
+				if (OutFailReason)
+				{
+					*OutFailReason = MoveTemp(ErrorMsg);
+				}
+				bSuccess = false;
 			}
-			return false;
 		}
-	}
 
-	return true;
+		UE_LOG(LogPluginUtils, Log, TEXT("Unloading assets from %d plugins took %0.2f sec"), PluginNames.Num(), FPlatformTime::Seconds() - StartTime);
+	}
+	return bSuccess;
 }
 
 bool FPluginUtils::UnloadPluginsAssets(const TConstArrayView<FString> PluginNames, FText* OutFailReason /*= nullptr*/)
