Additional error checking to prevent multiple .uplugin files existing when using the new plugin wizard. If a template .uplugin file is not simply named PLUGIN_NAME then it won't replace the new file that the wizard created. This results in the existence of multiple .uplugin files. Depending which file is found first it can cause a situation where a plugin with a different name is loaded and registered and this will cause the wizard to fail when it later searches for the registed plugin using the name that the user entered. There is a mismatch.

diff --git a/Engine/Plugins/Developer/PluginUtils/Source/PluginUtils/Private/PluginUtils.cpp b/Engine/Plugins/Developer/PluginUtils/Source/PluginUtils/Private/PluginUtils.cpp
--- a/Engine/Plugins/Developer/PluginUtils/Source/PluginUtils/Private/PluginUtils.cpp
+++ b/Engine/Plugins/Developer/PluginUtils/Source/PluginUtils/Private/PluginUtils.cpp
@@ -437,6 +437,27 @@ namespace PluginUtils
 
 		return Plugin;
 	}
+
+	bool ValidatePluginTemplate(const FString& TemplateFolder, FString& OutInvalidPluginTemplate)
+	{
+		constexpr const TCHAR PluginExt[] = TEXT(".uplugin");
+
+		TArray<FString> PluginTemplateFilenames;
+		IFileManager::Get().FindFiles(PluginTemplateFilenames, *TemplateFolder, PluginExt);
+
+		for (const FString& PluginTemplateFilename : PluginTemplateFilenames)
+		{
+			// The .uplugin filename can't have additional characters otherwise it won't match the user provided name and multiple .uplugin files will exist.
+			FString Filename = FPaths::GetBaseFilename(PluginTemplateFilename);
+			if (!Filename.Equals(PluginUtils::PLUGIN_NAME))
+			{
+				// Return first invalid.
+				OutInvalidPluginTemplate = PluginTemplateFilename;
+				return false;
+			}
+		}
+		return true;
+	}
 }
 
 FString FPluginUtils::GetPluginFolder(const FString& PluginLocation, const FString& PluginName, bool bFullPath)
@@ -612,6 +633,18 @@ TSharedPtr<IPlugin> FPluginUtils::CreateAndLoadNewPlugin(const FString& PluginNa
 			GWarn->BeginSlowTask(LOCTEXT("CopyingPluginTemplate", "Copying plugin template files..."), /*ShowProgressDialog*/ true, /*bShowCancelButton*/ false);
 			for (const FString& TemplateFolder : CreationParams.TemplateFolders)
 			{
+				FString InvalidPluginTemplateFilename;
+				if (!PluginUtils::ValidatePluginTemplate(*TemplateFolder, InvalidPluginTemplateFilename))
+				{
+					if (LoadParams.OutFailReason)
+					{
+						*LoadParams.OutFailReason = FText::Format(LOCTEXT("InvalidPluginTemplateFilename", "Failed to copy plugin template file {0} due to invalid filename. Doing so would result in the existence of multiple .uplugin files.\n\nPlease rename the template file to {1}.uplugin"),
+							FText::FromString(FPaths::ConvertRelativePathToFull(InvalidPluginTemplateFilename)), FText::FromString(PluginUtils::PLUGIN_NAME));
+					}
+					bSucceeded = false;
+					break;
+				}
+
 				if (!PluginUtils::CopyPluginTemplateFolder(*PluginFolder, *TemplateFolder, PluginName, NewFilePaths, LoadParams.OutFailReason))
 				{
 					if (LoadParams.OutFailReason)
