Copying plugin template files will setup temporary redirectors to fix up and internal package references during the new plugin creation and loading step. Also made a few changes to support a template path pointing to a real plugin. Special case handling of the name for an existing .uplugin or GameFeatureData.uasset when copying those files

#rnx

diff --git a/Engine/Plugins/Developer/PluginUtils/Source/PluginUtils/Private/PluginUtils.cpp b/Engine/Plugins/Developer/PluginUtils/Source/PluginUtils/Private/PluginUtils.cpp
--- a/Engine/Plugins/Developer/PluginUtils/Source/PluginUtils/Private/PluginUtils.cpp
+++ b/Engine/Plugins/Developer/PluginUtils/Source/PluginUtils/Private/PluginUtils.cpp
@@ -24,6 +24,7 @@
 #include "Misc/PackageName.h"
 #include "Misc/FeedbackContext.h"
 #include "PluginReferenceDescriptor.h"
+#include "UObject/CoreRedirects.h"
 #include "UObject/UObjectIterator.h"
 
 DEFINE_LOG_CATEGORY_STATIC(LogPluginUtils, Log, All);
@@ -47,7 +48,27 @@ namespace PluginUtils
 	}
 	PRAGMA_ENABLE_DEPRECATION_WARNINGS
 
-	bool CopyPluginTemplateFolder(const TCHAR* DestinationDirectory, const TCHAR* Source, const FString& PluginName, TArray<FString>& InOutFilePathsWritten, FText* OutFailReason = nullptr)
+	FString ConstructPackageName(const FString& RootPath, const TCHAR* Filename)
+	{
+		FString PackageName;
+
+		const FString ContentPath = FPaths::Combine(RootPath, TEXT("Content"));
+		if (FPaths::IsUnderDirectory(Filename, ContentPath))
+		{
+			FStringView RelativePkgPath;
+			if (FPathViews::TryMakeChildPathRelativeTo(Filename, ContentPath, RelativePkgPath))
+			{
+				RelativePkgPath = FPathViews::GetBaseFilenameWithPath(RelativePkgPath);
+				if (RelativePkgPath.Len() > 0 && !RelativePkgPath.EndsWith(TEXT("/")))
+				{
+					PackageName = FPaths::Combine(TEXT("/"), FPaths::GetBaseFilename(RootPath), RelativePkgPath);
+				}
+			}
+		}
+		return PackageName;
+	}
+
+	bool CopyPluginTemplateFolder(const TCHAR* DestinationDirectory, const TCHAR* Source, const FString& PluginName, TArray<FCoreRedirect>& InOutCoreRedirectList, TArray<FString>& InOutFilePathsWritten, FText* OutFailReason = nullptr)
 	{
 		check(DestinationDirectory);
 		check(Source);
@@ -92,14 +113,16 @@ namespace PluginUtils
 			TArray<FString> CopyUnmodifiedFileTypes;
 			TArray<FString>& FilePathsWritten;
 			FText* FailReason;
+			TArray<FCoreRedirect>& CoreRedirectList;
 
-			FCopyPluginFilesAndDirs(IPlatformFile& InPlatformFile, const TCHAR* InSourceRoot, const TCHAR* InDestRoot, const FString& InPluginName, TArray<FString>& InFilePathsWritten, FText* InFailReason)
+			FCopyPluginFilesAndDirs(IPlatformFile& InPlatformFile, const TCHAR* InSourceRoot, const TCHAR* InDestRoot, const FString& InPluginName, TArray<FString>& InFilePathsWritten, TArray<FCoreRedirect>& CoreRedirectList, FText* InFailReason)
 				: PlatformFile(InPlatformFile)
 				, SourceRoot(InSourceRoot)
 				, DestRoot(InDestRoot)
 				, PluginName(InPluginName)
 				, FilePathsWritten(InFilePathsWritten)
 				, FailReason(InFailReason)
+				, CoreRedirectList(CoreRedirectList)
 			{
 				// Which file types we want to replace instances of PLUGIN_NAME with the new Plugin Name
 				NameReplacementFileTypes.Add(TEXT("cs"));
@@ -154,6 +177,11 @@ namespace PluginUtils
 							NewName = FPaths::Combine(CopyToPath, CleanFilename);
 						}
 
+						// Redirect the template package name to the generated package name to fix up internal references.
+						FString SourcePackageName = ConstructPackageName(SourceRoot, FilenameOrDirectory);
+						FString DestPackageName = ConstructPackageName(DestRoot, *NewName);
+						CoreRedirectList.Add(FCoreRedirect(ECoreRedirectFlags::Type_Package, SourcePackageName, DestPackageName));
+
 						if (PlatformFile.FileExists(*NewName))
 						{
 							// Delete destination file if it exists
@@ -182,6 +210,14 @@ namespace PluginUtils
 
 							OutFileContents = OutFileContents.Replace(*PluginNameAPI, *PluginNameAPI.ToUpper(), ESearchCase::CaseSensitive);
 
+							// Special case the .uplugin as we may be copying a real plugin and the filename will not be named PLUGIN_NAME
+							// as it needs to be loaded by the editor. Ensure the new plugin name is used.
+							if (NewExt == TEXT("uplugin"))
+							{
+								FString CopyToPath = FPaths::GetPath(NewName);
+								NewName = FPaths::Combine(CopyToPath, PluginName) + TEXT(".uplugin");
+							}
+
 							if (!FFileHelper::SaveStringToFile(OutFileContents, *NewName))
 							{
 								if (FailReason)
@@ -212,7 +248,7 @@ namespace PluginUtils
 		};
 
 		// copy plugin files and directories visitor
-		FCopyPluginFilesAndDirs CopyFilesAndDirs(PlatformFile, *SourceDir, *DestDir, PluginName, InOutFilePathsWritten, OutFailReason);
+		FCopyPluginFilesAndDirs CopyFilesAndDirs(PlatformFile, *SourceDir, *DestDir, PluginName, InOutFilePathsWritten, InOutCoreRedirectList, OutFailReason);
 
 		// create all files subdirectories and files in subdirectories!
 		return PlatformFile.IterateDirectoryRecursively(*SourceDir, CopyFilesAndDirs);
@@ -283,16 +319,22 @@ namespace PluginUtils
 						{
 							const FString AssetName = AssetData.AssetName.ToString().Replace(*PLUGIN_NAME, *PluginName, ESearchCase::CaseSensitive);
 							const FString AssetPath = AssetData.PackagePath.ToString().Replace(*PLUGIN_NAME, *PluginName, ESearchCase::CaseSensitive);
-							AssetToOriginalFilePathMap.Add(AssetData.GetAsset(), File);
-							FAssetRenameData RenameData(AssetData.GetAsset(), AssetPath, AssetName);
-
-							AssetRenameData.Add(RenameData);
 
-							if (UObject* Asset = AssetData.GetAsset())
+							// When the new name and old name are an exact match the on disk asset will be missing. This is because the clean up step of RenameAssetsWithDialog deletes the old file and no new file will be created. 
+							// This behaviour is not supported by the existing editor workflows. Rather than change the functionality, the rename step can instead be skipped.
+							if (AssetName != AssetData.AssetName)
 							{
-								Asset->Modify();
-								TextNamespaceUtil::ClearPackageNamespace(Asset);
-								TextNamespaceUtil::EnsurePackageNamespace(Asset);
+								AssetToOriginalFilePathMap.Add(AssetData.GetAsset(), File);
+								FAssetRenameData RenameData(AssetData.GetAsset(), AssetPath, AssetName);
+
+								AssetRenameData.Add(RenameData);
+
+								if (UObject* Asset = AssetData.GetAsset())
+								{
+									Asset->Modify();
+									TextNamespaceUtil::ClearPackageNamespace(Asset);
+									TextNamespaceUtil::EnsurePackageNamespace(Asset);
+								}
 							}
 						}
 					}
@@ -570,6 +612,10 @@ TSharedPtr<IPlugin> FPluginUtils::CreateAndLoadNewPlugin(const FString& PluginNa
 
 	const FString PluginFolder = FPluginUtils::GetPluginFolder(PluginLocation, PluginName, /*bFullPath*/ true);
 
+	// This is to fix up any internal plugin references otherwise the generated plugin content could have
+	// package references back to the source template content.
+	TArray<FCoreRedirect> CoreRedirectList;
+
 	TSharedPtr<IPlugin> NewPlugin;
 	bool bSucceeded = true;
 	do
@@ -633,19 +679,7 @@ TSharedPtr<IPlugin> FPluginUtils::CreateAndLoadNewPlugin(const FString& PluginNa
 			GWarn->BeginSlowTask(LOCTEXT("CopyingPluginTemplate", "Copying plugin template files..."), /*ShowProgressDialog*/ true, /*bShowCancelButton*/ false);
 			for (const FString& TemplateFolder : CreationParams.TemplateFolders)
 			{
-				FString InvalidPluginTemplateFilename;
-				if (!PluginUtils::ValidatePluginTemplate(*TemplateFolder, InvalidPluginTemplateFilename))
-				{
-					if (LoadParams.OutFailReason)
-					{
-						*LoadParams.OutFailReason = FText::Format(LOCTEXT("InvalidPluginTemplateFilename", "Failed to copy plugin template file {0} due to invalid filename. Doing so would result in the existence of multiple .uplugin files.\n\nPlease rename the template file to {1}.uplugin"),
-							FText::FromString(FPaths::ConvertRelativePathToFull(InvalidPluginTemplateFilename)), FText::FromString(PluginUtils::PLUGIN_NAME));
-					}
-					bSucceeded = false;
-					break;
-				}
-
-				if (!PluginUtils::CopyPluginTemplateFolder(*PluginFolder, *TemplateFolder, PluginName, NewFilePaths, LoadParams.OutFailReason))
+				if (!PluginUtils::CopyPluginTemplateFolder(*PluginFolder, *TemplateFolder, PluginName, CoreRedirectList, NewFilePaths, LoadParams.OutFailReason))
 				{
 					if (LoadParams.OutFailReason)
 					{
@@ -663,6 +697,12 @@ TSharedPtr<IPlugin> FPluginUtils::CreateAndLoadNewPlugin(const FString& PluginNa
 			break;
 		}
 
+		if (!CoreRedirectList.IsEmpty())
+		{
+			// The package redirectors need to be in place before loading any content in the generated plugin.
+			FCoreRedirects::AddRedirectList(CoreRedirectList, TEXT("GenerateNewPlugin"));
+		}
+
 		// Compile plugin code
 		if (CreationParams.Descriptor.Modules.Num() > 0)
 		{
@@ -736,6 +776,12 @@ TSharedPtr<IPlugin> FPluginUtils::CreateAndLoadNewPlugin(const FString& PluginNa
 		}
 	} while (false);
 
+	if (!CoreRedirectList.IsEmpty())
+	{
+		// Remove the package redirectors now that all the generated content has been saved.
+		FCoreRedirects::RemoveRedirectList(CoreRedirectList, TEXT("GenerateNewPlugin"));
+	}
+
 	if (!bSucceeded)
 	{
 		// Delete the plugin folder is something goes wrong during the plugin creation.
