Allow setting the VerseVersion field on uplugin creation.

diff --git a/Engine/Plugins/Developer/PluginUtils/Source/PluginUtils/Private/PluginUtils.cpp b/Engine/Plugins/Developer/PluginUtils/Source/PluginUtils/Private/PluginUtils.cpp
--- a/Engine/Plugins/Developer/PluginUtils/Source/PluginUtils/Private/PluginUtils.cpp
+++ b/Engine/Plugins/Developer/PluginUtils/Source/PluginUtils/Private/PluginUtils.cpp
@@ -565,6 +565,7 @@ TSharedPtr<IPlugin> FPluginUtils::CreateAndLoadNewPlugin(const FString& PluginNa
 	ExCreationParams.Descriptor.EnabledByDefault = CreationParams.EnabledByDefault;
 	ExCreationParams.Descriptor.bExplicitlyLoaded = CreationParams.bExplicitelyLoaded;
 	ExCreationParams.Descriptor.VersePath = CreationParams.VersePath;
+	ExCreationParams.Descriptor.VerseVersion = CreationParams.VerseVersion;
 	ExCreationParams.Descriptor.bEnableVerseAssetReflection = CreationParams.bEnableVerseAssetReflection;
 
 	if (CreationParams.bHasModules)
diff --git a/Engine/Plugins/Developer/PluginUtils/Source/PluginUtils/Public/PluginUtils.h b/Engine/Plugins/Developer/PluginUtils/Source/PluginUtils/Public/PluginUtils.h
--- a/Engine/Plugins/Developer/PluginUtils/Source/PluginUtils/Public/PluginUtils.h
+++ b/Engine/Plugins/Developer/PluginUtils/Source/PluginUtils/Public/PluginUtils.h
@@ -91,6 +91,10 @@ public:
 		/** The Verse path to the root of this plugin's content directory */
 		FString VersePath;
 
+		/** The version of the Verse language that this plugin targets.
+			If no value is specified, the latest stable version is used. */
+		TOptional<uint32> VerseVersion;
+
 		/** If to generate Verse source code definitions from assets contained in this plugin */
 		bool bEnableVerseAssetReflection = false;
 
