Fixup issue where PLUGIN_NAME would replace more than desired causing odd extensions of PLUGIN_NAME_suffix_suffix.
Adding new core function for what should be replaced. Default behavior is to use the PluginName

diff --git a/Engine/Plugins/Developer/PluginUtils/Source/PluginUtils/Private/PluginUtils.cpp b/Engine/Plugins/Developer/PluginUtils/Source/PluginUtils/Private/PluginUtils.cpp
--- a/Engine/Plugins/Developer/PluginUtils/Source/PluginUtils/Private/PluginUtils.cpp
+++ b/Engine/Plugins/Developer/PluginUtils/Source/PluginUtils/Private/PluginUtils.cpp
@@ -68,7 +68,7 @@ namespace PluginUtils
 		return PackageName;
 	}
 
-	bool CopyPluginTemplateFolder(const TCHAR* DestinationDirectory, const TCHAR* Source, const FString& PluginName, TArray<FCoreRedirect>& InOutCoreRedirectList, TArray<FString>& InOutFilePathsWritten, FText* OutFailReason = nullptr)
+	bool CopyPluginTemplateFolder(const TCHAR* DestinationDirectory, const TCHAR* Source, const FString& PluginName, const FString& NameToReplace, TArray<FCoreRedirect>& InOutCoreRedirectList, TArray<FString>& InOutFilePathsWritten, FText* OutFailReason = nullptr)
 	{
 		check(DestinationDirectory);
 		check(Source);
@@ -108,6 +108,7 @@ namespace PluginUtils
 			const TCHAR* SourceRoot;
 			const TCHAR* DestRoot;
 			const FString& PluginName;
+			const FString& NameToReplace;
 			TArray<FString> NameReplacementFileTypes;
 			TArray<FString> IgnoredFileTypes;
 			TArray<FString> CopyUnmodifiedFileTypes;
@@ -115,11 +116,12 @@ namespace PluginUtils
 			FText* FailReason;
 			TArray<FCoreRedirect>& CoreRedirectList;
 
-			FCopyPluginFilesAndDirs(IPlatformFile& InPlatformFile, const TCHAR* InSourceRoot, const TCHAR* InDestRoot, const FString& InPluginName, TArray<FString>& InFilePathsWritten, TArray<FCoreRedirect>& CoreRedirectList, FText* InFailReason)
+			FCopyPluginFilesAndDirs(IPlatformFile& InPlatformFile, const TCHAR* InSourceRoot, const TCHAR* InDestRoot, const FString& InPluginName, const FString& InNameToReplace, TArray<FString>& InFilePathsWritten, TArray<FCoreRedirect>& CoreRedirectList, FText* InFailReason)
 				: PlatformFile(InPlatformFile)
 				, SourceRoot(InSourceRoot)
 				, DestRoot(InDestRoot)
 				, PluginName(InPluginName)
+				, NameToReplace(InNameToReplace)
 				, FilePathsWritten(InFilePathsWritten)
 				, FailReason(InFailReason)
 				, CoreRedirectList(CoreRedirectList)
@@ -147,7 +149,7 @@ namespace PluginUtils
 				FString NewName(FilenameOrDirectory);
 				// change the root and rename paths/files
 				NewName.RemoveFromStart(SourceRoot);
-				NewName = NewName.Replace(*PLUGIN_NAME, *PluginName, ESearchCase::CaseSensitive);
+				NewName = NewName.Replace(*PLUGIN_NAME, *NameToReplace, ESearchCase::CaseSensitive);
 				NewName = FPaths::Combine(DestRoot, *NewName);
 
 				if (bIsDirectory)
@@ -201,14 +203,14 @@ namespace PluginUtils
 								return false;
 							}
 
-							OutFileContents = OutFileContents.Replace(*PLUGIN_NAME, *PluginName, ESearchCase::CaseSensitive);
+							OutFileContents = OutFileContents.Replace(*PLUGIN_NAME, *NameToReplace, ESearchCase::CaseSensitive);
 
 							// For some content, we also want to export a PLUGIN_NAME_API text macro, which requires that the plugin name
 							// be all capitalized
 
-							FString PluginNameAPI = PluginName + TEXT("_API");
+							FString ReplaceNameAPI = NameToReplace + TEXT("_API");
 
-							OutFileContents = OutFileContents.Replace(*PluginNameAPI, *PluginNameAPI.ToUpper(), ESearchCase::CaseSensitive);
+							OutFileContents = OutFileContents.Replace(*ReplaceNameAPI, *ReplaceNameAPI.ToUpper(), ESearchCase::CaseSensitive);
 
 							// Special case the .uplugin as we may be copying a real plugin and the filename will not be named PLUGIN_NAME
 							// as it needs to be loaded by the editor. Ensure the new plugin name is used.
@@ -248,13 +250,13 @@ namespace PluginUtils
 		};
 
 		// copy plugin files and directories visitor
-		FCopyPluginFilesAndDirs CopyFilesAndDirs(PlatformFile, *SourceDir, *DestDir, PluginName, InOutFilePathsWritten, InOutCoreRedirectList, OutFailReason);
+		FCopyPluginFilesAndDirs CopyFilesAndDirs(PlatformFile, *SourceDir, *DestDir, PluginName, NameToReplace, InOutFilePathsWritten, InOutCoreRedirectList, OutFailReason);
 
 		// create all files subdirectories and files in subdirectories!
 		return PlatformFile.IterateDirectoryRecursively(*SourceDir, CopyFilesAndDirs);
 	}
 
-	void FixupPluginTemplateAssets(const FString& PluginName, TMap<FString, FString>& OutModifiedAssetPaths)
+	void FixupPluginTemplateAssets(const FString& PluginName, const FString& NameToReplace, TMap<FString, FString>& OutModifiedAssetPaths)
 	{
 		IPlatformFile& PlatformFile = FPlatformFileManager::Get().GetPlatformFile();
 
@@ -263,14 +265,14 @@ namespace PluginUtils
 		struct FFixupPluginAssets : public IPlatformFile::FDirectoryVisitor
 		{
 			IPlatformFile& PlatformFile;
-			const FString& PluginName;
+			const FString& NameToReplace;
 			const FString& PluginBaseDir;
 
 			TArray<FString> FilesToScan;
 
-			FFixupPluginAssets(IPlatformFile& InPlatformFile, const FString& InPluginName, const FString& InPluginBaseDir)
+			FFixupPluginAssets(IPlatformFile& InPlatformFile, const FString& InNameToReplace, const FString& InPluginBaseDir)
 				: PlatformFile(InPlatformFile)
-				, PluginName(InPluginName)
+				, NameToReplace(InNameToReplace)
 				, PluginBaseDir(InPluginBaseDir)
 			{
 			}
@@ -317,8 +319,8 @@ namespace PluginUtils
 
 						for (FAssetData AssetData : Assets)
 						{
-							const FString AssetName = AssetData.AssetName.ToString().Replace(*PLUGIN_NAME, *PluginName, ESearchCase::CaseSensitive);
-							const FString AssetPath = AssetData.PackagePath.ToString().Replace(*PLUGIN_NAME, *PluginName, ESearchCase::CaseSensitive);
+							const FString AssetName = AssetData.AssetName.ToString().Replace(*PLUGIN_NAME, *NameToReplace, ESearchCase::CaseSensitive);
+							const FString AssetPath = AssetData.PackagePath.ToString().Replace(*PLUGIN_NAME, *NameToReplace, ESearchCase::CaseSensitive);
 
 							// When the new name and old name are an exact match the on disk asset will be missing. This is because the clean up step of RenameAssetsWithDialog deletes the old file and no new file will be created. 
 							// This behaviour is not supported by the existing editor workflows. Rather than change the functionality, the rename step can instead be skipped.
@@ -371,7 +373,7 @@ namespace PluginUtils
 		if (Plugin.IsValid())
 		{
 			const FString PluginBaseDir = Plugin->GetBaseDir();
-			FFixupPluginAssets FixupPluginAssets(PlatformFile, PluginName, PluginBaseDir);
+			FFixupPluginAssets FixupPluginAssets(PlatformFile, NameToReplace, PluginBaseDir);
 			PlatformFile.IterateDirectoryRecursively(*PluginBaseDir, FixupPluginAssets);
 			FixupPluginAssets.PerformFixup(OutModifiedAssetPaths);
 		}
@@ -577,6 +579,11 @@ TSharedPtr<IPlugin> FPluginUtils::CreateAndLoadNewPlugin(const FString& PluginNa
 }
 
 TSharedPtr<IPlugin> FPluginUtils::CreateAndLoadNewPlugin(const FString& PluginName, const FString& PluginLocation, const FNewPluginParamsWithDescriptor& CreationParams, FLoadPluginParams& LoadParams)
+{
+	return CreateAndLoadNewPlugin(PluginName, PluginName, PluginLocation, CreationParams, LoadParams);
+}
+
+TSharedPtr<IPlugin> FPluginUtils::CreateAndLoadNewPlugin(const FString& PluginName, const FString& NameToReplace, const FString& PluginLocation, const FNewPluginParamsWithDescriptor& CreationParams, FLoadPluginParams& LoadParams)
 {
 	// Early validations on new plugin params
 	if (PluginName.IsEmpty())
@@ -680,7 +687,7 @@ TSharedPtr<IPlugin> FPluginUtils::CreateAndLoadNewPlugin(const FString& PluginNa
 			GWarn->BeginSlowTask(LOCTEXT("CopyingPluginTemplate", "Copying plugin template files..."), /*ShowProgressDialog*/ true, /*bShowCancelButton*/ false);
 			for (const FString& TemplateFolder : CreationParams.TemplateFolders)
 			{
-				if (!PluginUtils::CopyPluginTemplateFolder(*PluginFolder, *TemplateFolder, PluginName, CoreRedirectList, NewFilePaths, LoadParams.OutFailReason))
+				if (!PluginUtils::CopyPluginTemplateFolder(*PluginFolder, *TemplateFolder, PluginName, NameToReplace, CoreRedirectList, NewFilePaths, LoadParams.OutFailReason))
 				{
 					if (LoadParams.OutFailReason)
 					{
@@ -754,7 +761,7 @@ TSharedPtr<IPlugin> FPluginUtils::CreateAndLoadNewPlugin(const FString& PluginNa
 		{	
 			GWarn->BeginSlowTask(LOCTEXT("LoadingContent", "Loading Content..."), /*ShowProgressDialog*/ true, /*bShowCancelButton*/ false);
 			TMap<FString, FString> ModifiedAssetPaths;
-			PluginUtils::FixupPluginTemplateAssets(PluginName, ModifiedAssetPaths);
+			PluginUtils::FixupPluginTemplateAssets(PluginName, NameToReplace, ModifiedAssetPaths);
 
 			for (FString& CopiedFilePath : NewFilePaths)
 			{
diff --git a/Engine/Plugins/Developer/PluginUtils/Source/PluginUtils/Public/PluginUtils.h b/Engine/Plugins/Developer/PluginUtils/Source/PluginUtils/Public/PluginUtils.h
--- a/Engine/Plugins/Developer/PluginUtils/Source/PluginUtils/Public/PluginUtils.h
+++ b/Engine/Plugins/Developer/PluginUtils/Source/PluginUtils/Public/PluginUtils.h
@@ -195,7 +195,20 @@ public:
 	 * @note Will fail if the plugin already exists
 	 */
 	static TSharedPtr<IPlugin> CreateAndLoadNewPlugin(const FString& PluginName, const FString& PluginLocation, const FNewPluginParamsWithDescriptor& CreationParams, FLoadPluginParams& LoadParams);
-
+	
+	/**
+	 * Helper to create and load a new plugin
+	 * @param PluginName Plugin name
+	 * @param NameToReplace what will be used to replace PLUGIN_NAME in the given templates. Pass in PluginName here for default behavior
+	 * @param PluginLocation Directory that contains the plugin folder
+	 * @param CreationParams Plugin creation parameters
+	 * @param LoadParams Plugin loading parameters
+	 * @return The newly created plugin. If something goes wrong during the creation process, the plugin folder gets deleted and null is returned.
+	 * @note MountParams.OutFailReason outputs the reason the plugin creation or loading failed (if applicable)
+	 * @note Will fail if the plugin already exists
+	 */
+	static TSharedPtr<IPlugin> CreateAndLoadNewPlugin(const FString& PluginName, const FString& NameToReplace, const FString& PluginLocation, const FNewPluginParamsWithDescriptor& CreationParams, FLoadPluginParams& LoadParams);
+	
 	/**
 	 * Helper to create and mount a new plugin.
 	 * @param PluginName Plugin name
diff --git a/Engine/Plugins/Editor/PluginBrowser/Source/PluginBrowser/Private/SNewPluginWizard.cpp b/Engine/Plugins/Editor/PluginBrowser/Source/PluginBrowser/Private/SNewPluginWizard.cpp
--- a/Engine/Plugins/Editor/PluginBrowser/Source/PluginBrowser/Private/SNewPluginWizard.cpp
+++ b/Engine/Plugins/Editor/PluginBrowser/Source/PluginBrowser/Private/SNewPluginWizard.cpp
@@ -587,7 +587,7 @@ FText SNewPluginWizard::GetPluginDestinationPath() const
 FText SNewPluginWizard::GetCurrentPluginName() const
 {
 	TSharedPtr<FPluginTemplateDescription> Template = PluginWizardDefinition->GetSelectedTemplate();
-	if (Template.IsValid())
+	if (Template.IsValid() && !PluginNameText.IsEmpty())
 	{
 		TArray<FString> UPluginFiles;
 		IFileManager::Get().FindFiles(UPluginFiles, *Template->OnDiskPath, TEXT("*.uplugin"));
@@ -684,8 +684,9 @@ FReply SNewPluginWizard::OnCreatePluginClicked()
 	LoadParams.OutFailReason = &FailReason;
 
 	Template->CustomizeDescriptorBeforeCreation(CreationParams.Descriptor);
-	
-	TSharedPtr<IPlugin> NewPlugin = FPluginUtils::CreateAndLoadNewPlugin(PluginName, PluginFolderPath, CreationParams, LoadParams);
+
+
+	TSharedPtr<IPlugin> NewPlugin = FPluginUtils::CreateAndLoadNewPlugin(PluginName, PluginNameText.ToString(), PluginFolderPath, CreationParams, LoadParams);
 	const bool bSucceeded = NewPlugin.IsValid();
 
 
