Fixing an issue where multiple redirectors are created with the same source package name causing a redirector conflict. This can occur if two files have the same filename but the extensions are different (.csv, .uasset). Only need to create redirectors for unmodified files

diff --git a/Engine/Plugins/Developer/PluginUtils/Source/PluginUtils/Private/PluginUtils.cpp b/Engine/Plugins/Developer/PluginUtils/Source/PluginUtils/Private/PluginUtils.cpp
--- a/Engine/Plugins/Developer/PluginUtils/Source/PluginUtils/Private/PluginUtils.cpp
+++ b/Engine/Plugins/Developer/PluginUtils/Source/PluginUtils/Private/PluginUtils.cpp
@@ -177,12 +177,12 @@ namespace PluginUtils
 							FString CopyToPath = FPaths::GetPath(NewName);
 
 							NewName = FPaths::Combine(CopyToPath, CleanFilename);
-						}
 
-						// Redirect the template package name to the generated package name to fix up internal references.
-						FString SourcePackageName = ConstructPackageName(SourceRoot, FilenameOrDirectory);
-						FString DestPackageName = ConstructPackageName(DestRoot, *NewName);
-						CoreRedirectList.Add(FCoreRedirect(ECoreRedirectFlags::Type_Package, SourcePackageName, DestPackageName));
+							// Redirect the template package name to the generated package name to fix up internal references.
+							FString SourcePackageName = ConstructPackageName(SourceRoot, FilenameOrDirectory);
+							FString DestPackageName = ConstructPackageName(DestRoot, *NewName);
+							CoreRedirectList.Add(FCoreRedirect(ECoreRedirectFlags::Type_Package, SourcePackageName, DestPackageName));
+						}
 
 						if (PlatformFile.FileExists(*NewName))
 						{
@@ -1212,4 +1212,4 @@ bool FPluginUtils::IsValidPluginName(const FString& PluginName, FText* FailReaso
 
 #undef LOCTEXT_NAMESPACE
 
-#endif //if WITH_EDITOR
\ No newline at end of file
+#endif //if WITH_EDITOR
