Remove restrictions that a plugin name cannot start with a digit and that is has to pass IAssetTools::IsNameAllowed
These checks were introduced at a time where we thought the plugin name could be referenced in Verse code and consequently had to be a valid Verse identifier, but that's not true: only the plugin Verse path can

diff --git a/Engine/Plugins/Developer/PluginUtils/Source/PluginUtils/Private/PluginUtils.cpp b/Engine/Plugins/Developer/PluginUtils/Source/PluginUtils/Private/PluginUtils.cpp
--- a/Engine/Plugins/Developer/PluginUtils/Source/PluginUtils/Private/PluginUtils.cpp
+++ b/Engine/Plugins/Developer/PluginUtils/Source/PluginUtils/Private/PluginUtils.cpp
@@ -1169,16 +1169,6 @@ bool FPluginUtils::IsValidPluginName(const FString& PluginName, FText* FailReaso
 		return false;
 	}
 
-	// Must begin with an alphabetic character
-	if (!FChar::IsAlpha(PluginName[0]))
-	{
-		if (FailReason)
-		{
-			*FailReason = FText::Format(LOCTEXT("PluginNameMustBeginWithAlphabetic", "{0} name must begin with an alphabetic character"), PluginTermToUse);
-		}
-		return false;
-	}
-
 	// Only allow alphanumeric characters and underscore in the name
 	FString IllegalCharacters;
 	for (int32 CharIdx = 0; CharIdx < PluginName.Len(); ++CharIdx)
@@ -1202,12 +1192,6 @@ bool FPluginUtils::IsValidPluginName(const FString& PluginName, FText* FailReaso
 		return false;
 	}
 
-	FAssetToolsModule& AssetToolsModule = FModuleManager::LoadModuleChecked<FAssetToolsModule>("AssetTools");
-	if (!AssetToolsModule.Get().IsNameAllowed(PluginName, FailReason))
-	{
-		return false;
-	}
-
 	return true;
 }
 
diff --git a/Engine/Source/Developer/AssetTools/Public/IAssetTools.h b/Engine/Source/Developer/AssetTools/Public/IAssetTools.h
--- a/Engine/Source/Developer/AssetTools/Public/IAssetTools.h
+++ b/Engine/Source/Developer/AssetTools/Public/IAssetTools.h
@@ -622,9 +622,9 @@ public:
 	/** Returns true if all in list pass writable folder filter */
 	virtual bool AllPassWritableFolderFilter(const TArray<FString>& InPaths) const = 0;
 
-	/** Returns true if IsNameAllowedDelegate is not set, or if the name passes the filter function*/
+	/** Returns true if IsNameAllowedDelegate is not set, or if the name passes the filter function */
 	virtual bool IsNameAllowed(const FString& Name, FText* OutErrorMessage = nullptr) const = 0;
-	/** Allows setting of a global name filter that is applied to folders, assets, plugins, etc. */
+	/** Allows setting of a global name filter that is applied to assets and folders */
 	virtual void RegisterIsNameAllowedDelegate(const FName OwnerName, FIsNameAllowed Delegate) = 0;
 	/** Remove a previously-set global name filter */
 	virtual void UnregisterIsNameAllowedDelegate(const FName OwnerName) = 0;
