Add mount point check on plugin creation; there was a check against existing plugins but we also have to validate the plugin name isn't a mount point such as /Game, /Temp, /Script, etc

diff --git a/Engine/Plugins/Developer/PluginUtils/Source/PluginUtils/Private/PluginUtils.cpp b/Engine/Plugins/Developer/PluginUtils/Source/PluginUtils/Private/PluginUtils.cpp
--- a/Engine/Plugins/Developer/PluginUtils/Source/PluginUtils/Private/PluginUtils.cpp
+++ b/Engine/Plugins/Developer/PluginUtils/Source/PluginUtils/Private/PluginUtils.cpp
@@ -1151,6 +1151,23 @@ bool FPluginUtils::ValidateNewPluginNameAndLocation(const FString& PluginName, c
 		return false;
 	}
 
+	// Check if a mount point with the same name exists
+	{
+		FString MountPoint;
+		MountPoint.Reserve(PluginName.Len() + 2);
+		MountPoint.AppendChar(TEXT('/'));
+		MountPoint.Append(PluginName);
+		MountPoint.AppendChar(TEXT('/'));
+		if (FPackageName::MountPointExists(MountPoint))
+		{
+			if (FailReason)
+			{
+				*FailReason = FText::Format(LOCTEXT("MountPointExists", "Mount point {0} already exists"), FText::FromString(MountPoint));
+			}
+			return false;
+		}
+	}
+
 	return true;
 }
 
