FPluginUtils::FindPluginFromPackagePath that finds the plugin containing the specified package (if any) which supports C++ module packages (with the "/Script/" root)

diff --git a/Engine/Plugins/Developer/PluginUtils/Source/PluginUtils/Private/PluginUtils.cpp b/Engine/Plugins/Developer/PluginUtils/Source/PluginUtils/Private/PluginUtils.cpp
--- a/Engine/Plugins/Developer/PluginUtils/Source/PluginUtils/Private/PluginUtils.cpp
+++ b/Engine/Plugins/Developer/PluginUtils/Source/PluginUtils/Private/PluginUtils.cpp
@@ -871,6 +871,29 @@ TSharedPtr<IPlugin> FPluginUtils::FindLoadedPlugin(const FString& PluginDescript
 	return TSharedPtr<IPlugin>();
 }
 
+TSharedPtr<IPlugin> FPluginUtils::FindPluginFromPackagePath(FName PackagePath)
+{
+	FNameBuilder PackagePathNameBuilder(PackagePath);
+	return FindPluginFromPackagePath(PackagePathNameBuilder);
+}
+
+TSharedPtr<IPlugin> FPluginUtils::FindPluginFromPackagePath(FStringView PackagePath)
+{
+	TSharedPtr<IPlugin> Plugin;
+	static const FStringView ScriptPath = TEXTVIEW("/Script/");
+	if (PackagePath.StartsWith(ScriptPath))
+	{
+		const FName ModuleName(PackagePath.RightChop(ScriptPath.Len()));
+		Plugin = IPluginManager::Get().GetModuleOwnerPlugin(ModuleName);
+	}
+	else
+	{
+		const FStringView MountPoint = FPathViews::GetMountPointNameFromPath(PackagePath);
+		Plugin = IPluginManager::Get().FindPlugin(MountPoint);
+	}
+	return Plugin;
+}
+
 bool FPluginUtils::UnloadPlugin(const TSharedRef<IPlugin>& Plugin, FText* OutFailReason /*= nullptr*/)
 {
 	return UnloadPlugins({ Plugin }, OutFailReason);
diff --git a/Engine/Plugins/Developer/PluginUtils/Source/PluginUtils/Public/PluginUtils.h b/Engine/Plugins/Developer/PluginUtils/Source/PluginUtils/Public/PluginUtils.h
--- a/Engine/Plugins/Developer/PluginUtils/Source/PluginUtils/Public/PluginUtils.h
+++ b/Engine/Plugins/Developer/PluginUtils/Source/PluginUtils/Public/PluginUtils.h
@@ -300,6 +300,15 @@ public:
 	 */
 	static TSharedPtr<IPlugin> FindLoadedPlugin(const FString& PluginDescriptorFileName);
 
+	/**
+	 * Finds the plugin that contains the specified package, if any
+	 * @note This method supports C++ module packages (with the "/Script/" root)
+	 * @param PackagePath Unreal path of the package
+	 * @return Plugin or null if the package is not in a plugin
+	 */
+	static TSharedPtr<IPlugin> FindPluginFromPackagePath(FName PackagePath);
+	static TSharedPtr<IPlugin> FindPluginFromPackagePath(FStringView PackagePath);
+
 	/**
 	 * Unload assets from the specified plugin and unmount it
 	 * @note Only works on content-only plugins; plugins with code modules cannot be safely unloaded
