Overloaded PluginUtils functions to support full paths to plugin descriptors for situations where you want to sidestep plugin naming conventions.

diff --git a/Engine/Plugins/Developer/PluginUtils/Source/PluginUtils/Private/PluginUtils.cpp b/Engine/Plugins/Developer/PluginUtils/Source/PluginUtils/Private/PluginUtils.cpp
--- a/Engine/Plugins/Developer/PluginUtils/Source/PluginUtils/Private/PluginUtils.cpp
+++ b/Engine/Plugins/Developer/PluginUtils/Source/PluginUtils/Private/PluginUtils.cpp
@@ -383,7 +383,7 @@ namespace PluginUtils
 	{
 		LoadParams.bOutAlreadyLoaded = false;
 
-		if (LoadParams.bUpdateProjectPluginSearchPath)
+		if (LoadParams.bUpdateProjectPluginSearchPath && !PluginLocation.IsEmpty())
 		{
 			FPluginUtils::AddToPluginSearchPathIfNeeded(PluginLocation, /*bRefreshPlugins=*/false, /*bUpdateProjectFile=*/true);
 			IPluginManager::Get().RefreshPluginsList();
@@ -473,7 +473,7 @@ namespace PluginUtils
 			{
 				// Select plugin root folder in content browser
 				IContentBrowserSingleton& ContentBrowser = FModuleManager::LoadModuleChecked<FContentBrowserModule>("ContentBrowser").Get();
-				const bool bIsEnginePlugin = FPaths::IsUnderDirectory(PluginLocation, FPaths::EnginePluginsDir());
+				const bool bIsEnginePlugin = FPaths::IsUnderDirectory(PluginFilePath, FPaths::EnginePluginsDir());
 				ContentBrowser.ForceShowPluginContent(bIsEnginePlugin);
 				ContentBrowser.SetSelectedPaths({ PluginRootFolder }, /*bNeedsRefresh*/ true);
 			}
@@ -482,24 +482,309 @@ namespace PluginUtils
 		return Plugin;
 	}
 
-	bool ValidatePluginTemplate(const FString& TemplateFolder, FString& OutInvalidPluginTemplate)
+	void InitializeNewPluginParamsWithDescriptor(const FString& PluginName, const FPluginUtils::FNewPluginParams& CreationParams, FPluginUtils::FNewPluginParamsWithDescriptor& ExCreationParams)
 	{
-		constexpr const TCHAR PluginExt[] = TEXT(".uplugin");
+		ExCreationParams.PluginIconPath = CreationParams.PluginIconPath;
+		ExCreationParams.TemplateFolders = CreationParams.TemplateFolders;
+
+		ExCreationParams.Descriptor.FriendlyName = CreationParams.FriendlyName.Len() > 0 ? CreationParams.FriendlyName : PluginName;
+		ExCreationParams.Descriptor.Version = 1;
+		ExCreationParams.Descriptor.VersionName = TEXT("1.0");
+		ExCreationParams.Descriptor.Category = TEXT("Other");
+		ExCreationParams.Descriptor.CreatedBy = CreationParams.CreatedBy;
+		ExCreationParams.Descriptor.CreatedByURL = CreationParams.CreatedByURL;
+		ExCreationParams.Descriptor.Description = CreationParams.Description;
+		ExCreationParams.Descriptor.bIsBetaVersion = CreationParams.bIsBetaVersion;
+		ExCreationParams.Descriptor.bCanContainContent = CreationParams.bCanContainContent;
+		ExCreationParams.Descriptor.bCanContainVerse = CreationParams.bCanContainVerse;
+		ExCreationParams.Descriptor.EnabledByDefault = CreationParams.EnabledByDefault;
+		ExCreationParams.Descriptor.bExplicitlyLoaded = CreationParams.bExplicitelyLoaded;
+		ExCreationParams.Descriptor.VersePath = CreationParams.VersePath;
+		ExCreationParams.Descriptor.VerseVersion = CreationParams.VerseVersion;
+		ExCreationParams.Descriptor.bEnableSceneGraph = CreationParams.bEnableSceneGraph;
+		ExCreationParams.Descriptor.bEnableVerseAssetReflection = CreationParams.bEnableVerseAssetReflection;
+		ExCreationParams.Descriptor.bEnableIAD = CreationParams.bEnableIAD;
+
+		if (CreationParams.bHasModules)
+		{
+			ExCreationParams.Descriptor.Modules.Add(FModuleDescriptor(*PluginName, CreationParams.ModuleDescriptorType, CreationParams.LoadingPhase));
+		}
+	}
+
+	TSharedPtr<IPlugin> CreateAndLoadNewPluginInternal(const FString& PluginName, const FString& NameToReplace, const FString& PluginLocation, const FString& PluginFilePath, const FPluginUtils::FNewPluginParamsWithDescriptor& CreationParams, FPluginUtils::FLoadPluginParams& LoadParams)
+	{
+		if (!ensure(!PluginFilePath.IsEmpty()))
+		{
+			return nullptr;
+		}
+
+		if ((CreationParams.Descriptor.Modules.Num() > 0) && (CreationParams.TemplateFolders.Num() == 0))
+		{
+			if (LoadParams.OutFailReason)
+			{
+				*LoadParams.OutFailReason = LOCTEXT("CreateNewPluginParam_NoTemplateFolder", "A template folder must be specified to create a plugin with code");
+			}
+			return nullptr;
+		}
+
+		const FString PluginFolder = FPluginUtils::GetPluginFolder(PluginFilePath, /*bFullPath*/ true);
+
+		// This is to fix up any internal plugin references otherwise the generated plugin content could have
+		// package references back to the source template content.
+		TArray<FCoreRedirect> CoreRedirectList;
+
+		TSharedPtr<IPlugin> NewPlugin;
+		bool bSucceeded = true;
+		do
+		{
+			IPlatformFile& PlatformFile = FPlatformFileManager::Get().GetPlatformFile();
+			TArray<FString> NewFilePaths;
+
+			if (!PlatformFile.DirectoryExists(*PluginFolder) && !PlatformFile.CreateDirectoryTree(*PluginFolder))
+			{
+				if (LoadParams.OutFailReason)
+				{
+					*LoadParams.OutFailReason = FText::Format(LOCTEXT("FailedToCreatePluginFolder", "Failed to create plugin folder\n{0}"), FText::FromString(PluginFolder));
+				}
+				bSucceeded = false;
+				break;
+			}
+
+			if (CreationParams.Descriptor.bCanContainContent)
+			{
+				const FString PluginContentFolder = FPluginUtils::GetPluginContentFolder(PluginFilePath, /*bFullPath*/ true);
+				if (!PlatformFile.DirectoryExists(*PluginContentFolder) && !PlatformFile.CreateDirectory(*PluginContentFolder))
+				{
+					if (LoadParams.OutFailReason)
+					{
+						*LoadParams.OutFailReason = FText::Format(LOCTEXT("FailedToCreatePluginContentFolder", "Failed to create plugin Content folder\n{0}"), FText::FromString(PluginContentFolder));
+					}
+					bSucceeded = false;
+					break;
+				}
+			}
+
+			// Write the uplugin file
+			if (!CreationParams.Descriptor.Save(PluginFilePath, LoadParams.OutFailReason))
+			{
+				bSucceeded = false;
+				break;
+			}
+			NewFilePaths.Add(PluginFilePath);
+
+			// Copy plugin icon
+			if (!CreationParams.PluginIconPath.IsEmpty())
+			{
+				const FString ResourcesFolder = FPluginUtils::GetPluginResourcesFolder(PluginFilePath, /*bFullPath*/ true);
+				const FString DestinationPluginIconPath = FPaths::Combine(ResourcesFolder, TEXT("Icon128.png"));
+				if (IFileManager::Get().Copy(*DestinationPluginIconPath, *CreationParams.PluginIconPath, /*bReplaceExisting=*/ false) != COPY_OK)
+				{
+					if (LoadParams.OutFailReason)
+					{
+						*LoadParams.OutFailReason = FText::Format(LOCTEXT("FailedToCopyPluginIcon", "Failed to copy plugin icon\nFrom: {0}\nTo: {1}"), FText::FromString(FPaths::ConvertRelativePathToFull(CreationParams.PluginIconPath)), FText::FromString(DestinationPluginIconPath));
+					}
+					bSucceeded = false;
+					break;
+				}
+				NewFilePaths.Add(DestinationPluginIconPath);
+			}
+
+			// Copy template files
+			if (CreationParams.TemplateFolders.Num() > 0)
+			{
+				GWarn->BeginSlowTask(LOCTEXT("CopyingPluginTemplate", "Copying plugin template files..."), /*ShowProgressDialog*/ true, /*bShowCancelButton*/ false);
+				for (const FString& TemplateFolder : CreationParams.TemplateFolders)
+				{
+					if (!PluginUtils::CopyPluginTemplateFolder(*PluginFolder, *TemplateFolder, PluginName, NameToReplace, CoreRedirectList, NewFilePaths, LoadParams.OutFailReason))
+					{
+						if (LoadParams.OutFailReason)
+						{
+							*LoadParams.OutFailReason = FText::Format(LOCTEXT("FailedToCopyPluginTemplate", "Failed to copy plugin template files\nFrom: {0}\nTo: {1}\n{2}"), FText::FromString(FPaths::ConvertRelativePathToFull(TemplateFolder)), FText::FromString(PluginFolder), *LoadParams.OutFailReason);
+						}
+						bSucceeded = false;
+						break;
+					}
+				}
+				GWarn->EndSlowTask();
+			}
+
+			if (!bSucceeded)
+			{
+				break;
+			}
+
+			if (!CoreRedirectList.IsEmpty())
+			{
+				// The package redirectors need to be in place before loading any content in the generated plugin.
+				FCoreRedirects::AddRedirectList(CoreRedirectList, TEXT("GenerateNewPlugin"));
+			}
+
+			// Compile plugin code
+			if (CreationParams.Descriptor.Modules.Num() > 0)
+			{
+				const FString ProjectFileName = IFileManager::Get().ConvertToAbsolutePathForExternalAppForWrite(*FPaths::GetProjectFilePath());
+
+				// Internal plugins will be found and built automatically, and using the -Plugin param will mark them as external plugins causing them to fail to load.
+				const bool bIsEnginePlugin = FPaths::IsUnderDirectory(PluginFilePath, FPaths::EnginePluginsDir());
+				const bool bIsProjectPlugin = FPaths::IsUnderDirectory(PluginFilePath, FPaths::ProjectPluginsDir());
+				const FString LocalPluginArgument = (bIsEnginePlugin || bIsProjectPlugin) ? FString::Printf(TEXT("-BuildPluginAsLocal ")) : FString();
+
+				const FString Arguments = FString::Printf(TEXT("%s %s %s -Plugin=\"%s\" -Project=\"%s\" %s-Progress -NoHotReloadFromIDE"), FPlatformMisc::GetUBTTargetName(), FModuleManager::Get().GetUBTConfiguration(), FPlatformMisc::GetUBTPlatform(), *PluginFilePath, *ProjectFileName, *LocalPluginArgument);
+
+				if (!FDesktopPlatformModule::Get()->RunUnrealBuildTool(FText::Format(LOCTEXT("CompilingPlugin", "Compiling {0} plugin..."), FText::FromString(PluginName)), FPaths::RootDir(), Arguments, GWarn))
+				{
+					if (LoadParams.OutFailReason)
+					{
+						*LoadParams.OutFailReason = LOCTEXT("FailedToCompilePlugin", "Failed to compile plugin source code. See output log for more information.");
+					}
+					bSucceeded = false;
+					break;
+				}
+
+				// Reset the module paths cache. For unique build environments, the modules may be generated to the project binaries directory.
+				FModuleManager::Get().ResetModulePathsCache();
+
+				// Generate project files if we happen to be using a project file.
+				if (!FDesktopPlatformModule::Get()->GenerateProjectFiles(FPaths::RootDir(), FPaths::GetProjectFilePath(), GWarn))
+				{
+					if (LoadParams.OutFailReason)
+					{
+						*LoadParams.OutFailReason = LOCTEXT("FailedToGenerateProjectFiles", "Failed to generate project files");
+					}
+					bSucceeded = false;
+					break;
+				}
+			}
+
+			// Load the new plugin
+			NewPlugin = PluginUtils::LoadPluginInternal(PluginName, PluginLocation, PluginFilePath, LoadParams, /*bIsNewPlugin*/ true);
+			if (!NewPlugin)
+			{
+				bSucceeded = false;
+				break;
+			}
 
-		TArray<FString> PluginTemplateFilenames;
-		IFileManager::Get().FindFiles(PluginTemplateFilenames, *TemplateFolder, PluginExt);
+			// Fix any content that was added to the plugin
+			if (CreationParams.Descriptor.bCanContainContent)
+			{	
+				GWarn->BeginSlowTask(LOCTEXT("LoadingContent", "Loading Content..."), /*ShowProgressDialog*/ true, /*bShowCancelButton*/ false);
+				TMap<FString, FString> ModifiedAssetPaths;
+				PluginUtils::FixupPluginTemplateAssets(PluginName, NameToReplace, ModifiedAssetPaths);
+
+				for (FString& CopiedFilePath : NewFilePaths)
+				{
+					FString FullCopiedPath = FPaths::ConvertRelativePathToFull(CopiedFilePath);
+					FPaths::RemoveDuplicateSlashes(FullCopiedPath);
+					if (FString* NewPath = ModifiedAssetPaths.Find(FullCopiedPath))
+					{
+						CopiedFilePath = *NewPath;
+					}
+				}
+				GWarn->EndSlowTask();
+			}
 
-		for (const FString& PluginTemplateFilename : PluginTemplateFilenames)
+			// Add the plugin files to source control if the project is configured for it
+			if (USourceControlHelpers::IsAvailable())
+			{
+				GWarn->BeginSlowTask(LOCTEXT("AddingFilesToSourceControl", "Adding to Revision Control..."), /*ShowProgressDialog*/ true, /*bShowCancelButton*/ false);
+				USourceControlHelpers::MarkFilesForAdd(NewFilePaths);
+				GWarn->EndSlowTask();
+			}
+		} while (false);
+
+		if (!CoreRedirectList.IsEmpty())
 		{
-			// The .uplugin filename can't have additional characters otherwise it won't match the user provided name and multiple .uplugin files will exist.
-			FString Filename = FPaths::GetBaseFilename(PluginTemplateFilename);
-			if (!Filename.Equals(PluginUtils::PLUGIN_NAME))
+			// Remove the package redirectors now that all the generated content has been saved.
+			FCoreRedirects::RemoveRedirectList(CoreRedirectList, TEXT("GenerateNewPlugin"));
+		}
+
+		if (!bSucceeded)
+		{
+			// Delete the plugin folder is something goes wrong during the plugin creation.
+			IFileManager::Get().DeleteDirectory(*PluginFolder, /*RequireExists*/ false, /*Tree*/ true);
+			if (NewPlugin)
 			{
-				// Return first invalid.
-				OutInvalidPluginTemplate = PluginTemplateFilename;
+				// Refresh plugins if the new plugin was registered, but we decide to delete its files.
+				IPluginManager::Get().RefreshPluginsList();
+				NewPlugin.Reset();
+				ensure(!IPluginManager::Get().FindPlugin(PluginName));
+			}
+		}
+
+		return NewPlugin;
+	}
+
+	bool ValidateNewPluginFilePath(const FString& PluginFilePath, FText* FailReason)
+	{
+		// Check if a .uplugin file exists
+		if (!PluginFilePath.IsEmpty() && FPaths::FileExists(PluginFilePath))
+		{
+			if (FailReason)
+			{
+				*FailReason = FText::Format(LOCTEXT("PluginPathExists", "Plugin already exists at this location\n{0}"), FText::FromString(FPaths::ConvertRelativePathToFull(PluginFilePath)));
+			}
+			return false;
+		}
+
+		return true;
+	}
+
+	bool ValidateNewPluginFolder(const FString& PluginFolder, FText* FailReason)
+	{
+		// Check there isn't an existing file along the plugin folder path that would prevent creating the directory tree
+		{
+			FString ExistingFilePath = PluginFolder;
+			while (!ExistingFilePath.IsEmpty())
+			{
+				if (FPaths::FileExists(ExistingFilePath))
+				{
+					break;
+				}
+				ExistingFilePath = FPaths::GetPath(ExistingFilePath);
+			}
+
+			if (!ExistingFilePath.IsEmpty())
+			{
+				if (FailReason)
+				{
+					*FailReason = FText::Format(LOCTEXT("PluginLocationIsFile", "Plugin location is invalid because the following file is in the way\n{0}"), FText::FromString(ExistingFilePath));
+				}
+				return false;
+			}
+		}
+
+		return true;
+	}
+
+	bool ValidateNewPluginName(const FString& PluginName, FText* FailReason)
+	{
+		// Check to see if a discovered plugin with this name exists (at any path)
+		if (TSharedPtr<IPlugin> ExistingPlugin = IPluginManager::Get().FindPlugin(PluginName))
+		{
+			if (FailReason)
+			{
+				*FailReason = FText::Format(LOCTEXT("PluginNameAlreadyInUse", "Plugin name is already in use by\n{0}"), FText::FromString(FPaths::ConvertRelativePathToFull(ExistingPlugin->GetDescriptorFileName())));
+			}
+			return false;
+		}
+
+		// Check if a mount point with the same name exists
+		{
+			FString MountPoint;
+			MountPoint.Reserve(PluginName.Len() + 2);
+			MountPoint.AppendChar(TEXT('/'));
+			MountPoint.Append(PluginName);
+			MountPoint.AppendChar(TEXT('/'));
+			if (FPackageName::MountPointExists(MountPoint))
+			{
+				if (FailReason)
+				{
+					*FailReason = FText::Format(LOCTEXT("MountPointExists", "Mount point {0} already exists"), FText::FromString(MountPoint));
+				}
 				return false;
 			}
 		}
+
 		return true;
 	}
 }
@@ -515,6 +800,17 @@ FString FPluginUtils::GetPluginFolder(const FString& PluginLocation, const FStri
 	return PluginFolder;
 }
 
+FString FPluginUtils::GetPluginFolder(const FString& PluginFilePath, bool bFullPath)
+{
+	FString PluginFolder = FPaths::GetPath(PluginFilePath);
+	if (bFullPath)
+	{
+		PluginFolder = FPaths::ConvertRelativePathToFull(PluginFolder);
+	}
+	FPaths::MakePlatformFilename(PluginFolder);
+	return PluginFolder;
+}
+
 FString FPluginUtils::GetPluginFilePath(const FString& PluginLocation, const FString& PluginName, bool bFullPath)
 {
 	FString PluginFilePath = FPaths::Combine(PluginLocation, PluginName, (PluginName + FPluginDescriptor::GetFileExtension()));
@@ -526,6 +822,11 @@ FString FPluginUtils::GetPluginFilePath(const FString& PluginLocation, const FSt
 	return PluginFilePath;
 }
 
+FString FPluginUtils::GetPluginName(const FString& PluginFilePath)
+{
+	return FPaths::GetBaseFilename(PluginFilePath);
+}
+
 FString FPluginUtils::GetPluginContentFolder(const FString& PluginLocation, const FString& PluginName, bool bFullPath)
 {
 	FString PluginContentFolder = FPaths::Combine(PluginLocation, PluginName, TEXT("Content"));
@@ -537,6 +838,17 @@ FString FPluginUtils::GetPluginContentFolder(const FString& PluginLocation, cons
 	return PluginContentFolder;
 }
 
+FString FPluginUtils::GetPluginContentFolder(const FString& PluginFilePath, bool bFullPath)
+{
+	FString PluginContentFolder = FPaths::Combine(FPaths::GetPath(PluginFilePath), TEXT("Content"));
+	if (bFullPath)
+	{
+		PluginContentFolder = FPaths::ConvertRelativePathToFull(PluginContentFolder);
+	}
+	FPaths::MakePlatformFilename(PluginContentFolder);
+	return PluginContentFolder;
+}
+
 FString FPluginUtils::GetPluginResourcesFolder(const FString& PluginLocation, const FString& PluginName, bool bFullPath)
 {
 	FString PluginResourcesFolder = FPaths::Combine(PluginLocation, PluginName, TEXT("Resources"));
@@ -548,41 +860,30 @@ FString FPluginUtils::GetPluginResourcesFolder(const FString& PluginLocation, co
 	return PluginResourcesFolder;
 }
 
-TSharedPtr<IPlugin> FPluginUtils::CreateAndLoadNewPlugin(const FString& PluginName, const FString& PluginLocation, const FNewPluginParams& CreationParams, FLoadPluginParams& LoadParams)
+FString FPluginUtils::GetPluginResourcesFolder(const FString& PluginFilePath, bool bFullPath)
 {
-	FNewPluginParamsWithDescriptor ExCreationParams;
-	ExCreationParams.PluginIconPath = CreationParams.PluginIconPath;
-	ExCreationParams.TemplateFolders = CreationParams.TemplateFolders;
-
-	ExCreationParams.Descriptor.FriendlyName = CreationParams.FriendlyName.Len() > 0 ? CreationParams.FriendlyName : PluginName;
-	ExCreationParams.Descriptor.Version = 1;
-	ExCreationParams.Descriptor.VersionName = TEXT("1.0");
-	ExCreationParams.Descriptor.Category = TEXT("Other");
-	ExCreationParams.Descriptor.CreatedBy = CreationParams.CreatedBy;
-	ExCreationParams.Descriptor.CreatedByURL = CreationParams.CreatedByURL;
-	ExCreationParams.Descriptor.Description = CreationParams.Description;
-	ExCreationParams.Descriptor.bIsBetaVersion = CreationParams.bIsBetaVersion;
-	ExCreationParams.Descriptor.bCanContainContent = CreationParams.bCanContainContent;
-	ExCreationParams.Descriptor.bCanContainVerse = CreationParams.bCanContainVerse;
-	ExCreationParams.Descriptor.EnabledByDefault = CreationParams.EnabledByDefault;
-	ExCreationParams.Descriptor.bExplicitlyLoaded = CreationParams.bExplicitelyLoaded;
-	ExCreationParams.Descriptor.VersePath = CreationParams.VersePath;
-	ExCreationParams.Descriptor.VerseVersion = CreationParams.VerseVersion;
-	ExCreationParams.Descriptor.bEnableSceneGraph = CreationParams.bEnableSceneGraph;
-	ExCreationParams.Descriptor.bEnableVerseAssetReflection = CreationParams.bEnableVerseAssetReflection;
-	ExCreationParams.Descriptor.bEnableIAD = CreationParams.bEnableIAD;
-
-	if (CreationParams.bHasModules)
+	FString PluginResourcesFolder = FPaths::Combine(FPaths::GetPath(PluginFilePath), TEXT("Resources"));
+	if (bFullPath)
 	{
-		ExCreationParams.Descriptor.Modules.Add(FModuleDescriptor(*PluginName, CreationParams.ModuleDescriptorType, CreationParams.LoadingPhase));
+		PluginResourcesFolder = FPaths::ConvertRelativePathToFull(PluginResourcesFolder);
 	}
+	FPaths::MakePlatformFilename(PluginResourcesFolder);
+	return PluginResourcesFolder;
+}
+
+TSharedPtr<IPlugin> FPluginUtils::CreateAndLoadNewPlugin(const FString& PluginName, const FString& PluginLocation, const FNewPluginParams& CreationParams, FLoadPluginParams& LoadParams)
+{
+	FPluginUtils::FNewPluginParamsWithDescriptor ExCreationParams;
+	PluginUtils::InitializeNewPluginParamsWithDescriptor(PluginName, CreationParams, ExCreationParams);
 
 	return CreateAndLoadNewPlugin(PluginName, PluginLocation, ExCreationParams, LoadParams);
 }
 
 TSharedPtr<IPlugin> FPluginUtils::CreateAndLoadNewPlugin(const FString& PluginName, const FString& PluginLocation, const FNewPluginParamsWithDescriptor& CreationParams, FLoadPluginParams& LoadParams)
 {
-	return CreateAndLoadNewPlugin(PluginName, PluginName, PluginLocation, CreationParams, LoadParams);
+	PRAGMA_DISABLE_DEPRECATION_WARNINGS
+	return CreateAndLoadNewPlugin(PluginName, !CreationParams.NameToReplace.IsEmpty() ? CreationParams.NameToReplace : PluginName, PluginLocation, CreationParams, LoadParams);
+	PRAGMA_ENABLE_DEPRECATION_WARNINGS
 }
 
 TSharedPtr<IPlugin> FPluginUtils::CreateAndLoadNewPlugin(const FString& PluginName, const FString& NameToReplace, const FString& PluginLocation, const FNewPluginParamsWithDescriptor& CreationParams, FLoadPluginParams& LoadParams)
@@ -606,206 +907,47 @@ TSharedPtr<IPlugin> FPluginUtils::CreateAndLoadNewPlugin(const FString& PluginNa
 		return nullptr;
 	}
 
-	if ((CreationParams.Descriptor.Modules.Num() > 0) && (CreationParams.TemplateFolders.Num() == 0))
-	{
-		if (LoadParams.OutFailReason)
-		{
-			*LoadParams.OutFailReason = LOCTEXT("CreateNewPluginParam_NoTemplateFolder", "A template folder must be specified to create a plugin with code");
-		}
-		return nullptr;
-	}
-
 	if (!FPluginUtils::ValidateNewPluginNameAndLocation(PluginName, PluginLocation, LoadParams.OutFailReason))
 	{
 		return nullptr;
 	}
 
-	const FString PluginFolder = FPluginUtils::GetPluginFolder(PluginLocation, PluginName, /*bFullPath*/ true);
-
-	// This is to fix up any internal plugin references otherwise the generated plugin content could have
-	// package references back to the source template content.
-	TArray<FCoreRedirect> CoreRedirectList;
-
-	TSharedPtr<IPlugin> NewPlugin;
-	bool bSucceeded = true;
-	do
-	{
-		IPlatformFile& PlatformFile = FPlatformFileManager::Get().GetPlatformFile();
-		TArray<FString> NewFilePaths;
-
-		if (!PlatformFile.DirectoryExists(*PluginFolder) && !PlatformFile.CreateDirectoryTree(*PluginFolder))
-		{
-			if (LoadParams.OutFailReason)
-			{
-				*LoadParams.OutFailReason = FText::Format(LOCTEXT("FailedToCreatePluginFolder", "Failed to create plugin folder\n{0}"), FText::FromString(PluginFolder));
-			}
-			bSucceeded = false;
-			break;
-		}
-
-		if (CreationParams.Descriptor.bCanContainContent)
-		{
-			const FString PluginContentFolder = FPluginUtils::GetPluginContentFolder(PluginLocation, PluginName, /*bFullPath*/ true);
-			if (!PlatformFile.DirectoryExists(*PluginContentFolder) && !PlatformFile.CreateDirectory(*PluginContentFolder))
-			{
-				if (LoadParams.OutFailReason)
-				{
-					*LoadParams.OutFailReason = FText::Format(LOCTEXT("FailedToCreatePluginContentFolder", "Failed to create plugin Content folder\n{0}"), FText::FromString(PluginContentFolder));
-				}
-				bSucceeded = false;
-				break;
-			}
-		}
-
-		// Write the uplugin file
-		const FString PluginFilePath = FPluginUtils::GetPluginFilePath(PluginLocation, PluginName, /*bFullPath*/ true);
-		if (!CreationParams.Descriptor.Save(PluginFilePath, LoadParams.OutFailReason))
-		{
-			bSucceeded = false;
-			break;
-		}
-		NewFilePaths.Add(PluginFilePath);
-
-		// Copy plugin icon
-		if (!CreationParams.PluginIconPath.IsEmpty())
-		{
-			const FString ResourcesFolder = FPluginUtils::GetPluginResourcesFolder(PluginLocation, PluginName, /*bFullPath*/ true);
-			const FString DestinationPluginIconPath = FPaths::Combine(ResourcesFolder, TEXT("Icon128.png"));
-			if (IFileManager::Get().Copy(*DestinationPluginIconPath, *CreationParams.PluginIconPath, /*bReplaceExisting=*/ false) != COPY_OK)
-			{
-				if (LoadParams.OutFailReason)
-				{
-					*LoadParams.OutFailReason = FText::Format(LOCTEXT("FailedToCopyPluginIcon", "Failed to copy plugin icon\nFrom: {0}\nTo: {1}"), FText::FromString(FPaths::ConvertRelativePathToFull(CreationParams.PluginIconPath)), FText::FromString(DestinationPluginIconPath));
-				}
-				bSucceeded = false;
-				break;
-			}
-			NewFilePaths.Add(DestinationPluginIconPath);
-		}
-
-		// Copy template files
-		if (CreationParams.TemplateFolders.Num() > 0)
-		{
-			GWarn->BeginSlowTask(LOCTEXT("CopyingPluginTemplate", "Copying plugin template files..."), /*ShowProgressDialog*/ true, /*bShowCancelButton*/ false);
-			for (const FString& TemplateFolder : CreationParams.TemplateFolders)
-			{
-				if (!PluginUtils::CopyPluginTemplateFolder(*PluginFolder, *TemplateFolder, PluginName, NameToReplace, CoreRedirectList, NewFilePaths, LoadParams.OutFailReason))
-				{
-					if (LoadParams.OutFailReason)
-					{
-						*LoadParams.OutFailReason = FText::Format(LOCTEXT("FailedToCopyPluginTemplate", "Failed to copy plugin template files\nFrom: {0}\nTo: {1}\n{2}"), FText::FromString(FPaths::ConvertRelativePathToFull(TemplateFolder)), FText::FromString(PluginFolder), *LoadParams.OutFailReason);
-					}
-					bSucceeded = false;
-					break;
-				}
-			}
-			GWarn->EndSlowTask();
-		}
-
-		if (!bSucceeded)
-		{
-			break;
-		}
-
-		if (!CoreRedirectList.IsEmpty())
-		{
-			// The package redirectors need to be in place before loading any content in the generated plugin.
-			FCoreRedirects::AddRedirectList(CoreRedirectList, TEXT("GenerateNewPlugin"));
-		}
-
-		// Compile plugin code
-		if (CreationParams.Descriptor.Modules.Num() > 0)
-		{
-			const FString ProjectFileName = IFileManager::Get().ConvertToAbsolutePathForExternalAppForWrite(*FPaths::GetProjectFilePath());
-
-			// Internal plugins will be found and built automatically, and using the -Plugin param will mark them as external plugins causing them to fail to load.
-			const bool bIsEnginePlugin = FPaths::IsUnderDirectory(PluginFilePath, FPaths::EnginePluginsDir());
-			const bool bIsProjectPlugin = FPaths::IsUnderDirectory(PluginFilePath, FPaths::ProjectPluginsDir());
-			const FString LocalPluginArgument = (bIsEnginePlugin || bIsProjectPlugin) ? FString::Printf(TEXT("-BuildPluginAsLocal ")) : FString();
-
-			const FString Arguments = FString::Printf(TEXT("%s %s %s -Plugin=\"%s\" -Project=\"%s\" %s-Progress -NoHotReloadFromIDE"), FPlatformMisc::GetUBTTargetName(), FModuleManager::Get().GetUBTConfiguration(), FPlatformMisc::GetUBTPlatform(), *PluginFilePath, *ProjectFileName, *LocalPluginArgument);
-
-			if (!FDesktopPlatformModule::Get()->RunUnrealBuildTool(FText::Format(LOCTEXT("CompilingPlugin", "Compiling {0} plugin..."), FText::FromString(PluginName)), FPaths::RootDir(), Arguments, GWarn))
-			{
-				if (LoadParams.OutFailReason)
-				{
-					*LoadParams.OutFailReason = LOCTEXT("FailedToCompilePlugin", "Failed to compile plugin source code. See output log for more information.");
-				}
-				bSucceeded = false;
-				break;
-			}
-
-			// Reset the module paths cache. For unique build environments, the modules may be generated to the project binaries directory.
-			FModuleManager::Get().ResetModulePathsCache();
-
-			// Generate project files if we happen to be using a project file.
-			if (!FDesktopPlatformModule::Get()->GenerateProjectFiles(FPaths::RootDir(), FPaths::GetProjectFilePath(), GWarn))
-			{
-				if (LoadParams.OutFailReason)
-				{
-					*LoadParams.OutFailReason = LOCTEXT("FailedToGenerateProjectFiles", "Failed to generate project files");
-				}
-				bSucceeded = false;
-				break;
-			}
-		}
+	const FString PluginFilePath = FPluginUtils::GetPluginFilePath(PluginLocation, PluginName, /*bFullPath*/ true);
 
-		// Load the new plugin
-		NewPlugin = PluginUtils::LoadPluginInternal(PluginName, PluginLocation, PluginFilePath, LoadParams, /*bIsNewPlugin*/ true);
-		if (!NewPlugin)
-		{
-			bSucceeded = false;
-			break;
-		}
+	return PluginUtils::CreateAndLoadNewPluginInternal(PluginName, NameToReplace, PluginLocation, PluginFilePath, CreationParams, LoadParams);
+}
 
-		// Fix any content that was added to the plugin
-		if (CreationParams.Descriptor.bCanContainContent)
-		{	
-			GWarn->BeginSlowTask(LOCTEXT("LoadingContent", "Loading Content..."), /*ShowProgressDialog*/ true, /*bShowCancelButton*/ false);
-			TMap<FString, FString> ModifiedAssetPaths;
-			PluginUtils::FixupPluginTemplateAssets(PluginName, NameToReplace, ModifiedAssetPaths);
+TSharedPtr<IPlugin> FPluginUtils::CreateAndLoadNewPlugin(const FString& PluginFilePath, const FNewPluginParams& CreationParams, FLoadPluginParams& LoadParams)
+{
+	FPluginUtils::FNewPluginParamsWithDescriptor ExCreationParams;
+	PluginUtils::InitializeNewPluginParamsWithDescriptor(GetPluginName(PluginFilePath), CreationParams, ExCreationParams);
 
-			for (FString& CopiedFilePath : NewFilePaths)
-			{
-				FString FullCopiedPath = FPaths::ConvertRelativePathToFull(CopiedFilePath);
-				FPaths::RemoveDuplicateSlashes(FullCopiedPath);
-				if (FString* NewPath = ModifiedAssetPaths.Find(FullCopiedPath))
-				{
-					CopiedFilePath = *NewPath;
-				}
-			}
-			GWarn->EndSlowTask();
-		}
+	return CreateAndLoadNewPlugin(PluginFilePath, ExCreationParams, LoadParams);
+}
 
-		// Add the plugin files to source control if the project is configured for it
-		if (USourceControlHelpers::IsAvailable())
+TSharedPtr<IPlugin> FPluginUtils::CreateAndLoadNewPlugin(const FString& PluginFilePath, const FNewPluginParamsWithDescriptor& CreationParams, FLoadPluginParams& LoadParams)
+{
+	// Early validations on new plugin params
+	if (PluginFilePath.IsEmpty())
+	{
+		if (LoadParams.OutFailReason)
 		{
-			GWarn->BeginSlowTask(LOCTEXT("AddingFilesToSourceControl", "Adding to Revision Control..."), /*ShowProgressDialog*/ true, /*bShowCancelButton*/ false);
-			USourceControlHelpers::MarkFilesForAdd(NewFilePaths);
-			GWarn->EndSlowTask();
+			*LoadParams.OutFailReason = LOCTEXT("CreateNewPluginParam_NoPluginFilePath", "Missing plugin file name");
 		}
-	} while (false);
-
-	if (!CoreRedirectList.IsEmpty())
-	{
-		// Remove the package redirectors now that all the generated content has been saved.
-		FCoreRedirects::RemoveRedirectList(CoreRedirectList, TEXT("GenerateNewPlugin"));
+		return nullptr;
 	}
 
-	if (!bSucceeded)
+	if (!FPluginUtils::ValidateNewPluginFilePath(PluginFilePath, LoadParams.OutFailReason))
 	{
-		// Delete the plugin folder is something goes wrong during the plugin creation.
-		IFileManager::Get().DeleteDirectory(*PluginFolder, /*RequireExists*/ false, /*Tree*/ true);
-		if (NewPlugin)
-		{
-			// Refresh plugins if the new plugin was registered, but we decide to delete its files.
-			IPluginManager::Get().RefreshPluginsList();
-			NewPlugin.Reset();
-			ensure(!IPluginManager::Get().FindPlugin(PluginName));
-		}
+		return nullptr;
 	}
 
-	return NewPlugin;
+	const FString PluginName = GetPluginName(PluginFilePath);
+
+	FString FullPluginFilePath = FPaths::ConvertRelativePathToFull(PluginFilePath);
+	FPaths::MakePlatformFilename(FullPluginFilePath);
+
+	return PluginUtils::CreateAndLoadNewPluginInternal(PluginName, !CreationParams.NameToReplace.IsEmpty() ? CreationParams.NameToReplace : PluginName, FString(), FullPluginFilePath, CreationParams, LoadParams);
 }
 
 PRAGMA_DISABLE_DEPRECATION_WARNINGS
@@ -829,26 +971,26 @@ TSharedPtr<IPlugin> FPluginUtils::LoadPlugin(const FString& PluginName, const FS
 	return LoadPlugin(FPluginUtils::GetPluginFilePath(PluginLocation, PluginName), LoadParams);
 }
 
-TSharedPtr<IPlugin> FPluginUtils::LoadPlugin(const FString& PluginFileName, FLoadPluginParams& LoadParams)
+TSharedPtr<IPlugin> FPluginUtils::LoadPlugin(const FString& PluginFilePath, FLoadPluginParams& LoadParams)
 {
 	// Valide that the uplugin file exists.
-	if (!FPaths::FileExists(PluginFileName))
+	if (!FPaths::FileExists(PluginFilePath))
 	{
 		if (LoadParams.OutFailReason)
 		{
-			*LoadParams.OutFailReason = FText::Format(LOCTEXT("PluginFileDoesNotExist", "Plugin file does not exist\n{0}"), FText::FromString(FPaths::ConvertRelativePathToFull(PluginFileName)));
+			*LoadParams.OutFailReason = FText::Format(LOCTEXT("PluginFileDoesNotExist", "Plugin file does not exist\n{0}"), FText::FromString(FPaths::ConvertRelativePathToFull(PluginFilePath)));
 		}
 		return nullptr;
 	}
 
-	const FString PluginName = FPaths::GetBaseFilename(PluginFileName);
+	const FString PluginName = GetPluginName(PluginFilePath);
 	if (!IsValidPluginName(PluginName, LoadParams.OutFailReason))
 	{
 		return nullptr;
 	}
 
-	const FString PluginLocation = FPaths::GetPath(FPaths::GetPath(PluginFileName));
-	return PluginUtils::LoadPluginInternal(PluginName, PluginLocation, PluginFileName, LoadParams, /*bIsNewPlugin*/ false);
+	const FString PluginLocation = FPaths::GetPath(FPaths::GetPath(PluginFilePath));
+	return PluginUtils::LoadPluginInternal(PluginName, PluginLocation, PluginFilePath, LoadParams, /*bIsNewPlugin*/ false);
 }
 
 PRAGMA_DISABLE_DEPRECATION_WARNINGS
@@ -1118,18 +1260,9 @@ bool FPluginUtils::ValidateNewPluginNameAndLocation(const FString& PluginName, c
 
 	if (!PluginLocation.IsEmpty())
 	{
-		// Check if a .uplugin file exists at the specified location (if any)
+		if (!PluginUtils::ValidateNewPluginFilePath(FPluginUtils::GetPluginFilePath(PluginLocation, PluginName), FailReason))
 		{
-			const FString PluginFilePath = FPluginUtils::GetPluginFilePath(PluginLocation, PluginName);
-
-			if (!PluginFilePath.IsEmpty() && FPaths::FileExists(PluginFilePath))
-			{
-				if (FailReason)
-				{
-					*FailReason = FText::Format(LOCTEXT("PluginPathExists", "Plugin already exists at this location\n{0}"), FText::FromString(FPaths::ConvertRelativePathToFull(PluginFilePath)));
-				}
-				return false;
-			}
+			return false;
 		}
 
 		// Check that the plugin location is a valid path (it doesn't have to exist; it will be created if needed)
@@ -1142,54 +1275,43 @@ bool FPluginUtils::ValidateNewPluginNameAndLocation(const FString& PluginName, c
 			return false;
 		}
 
-		// Check there isn't an existing file along the plugin folder path that would prevent creating the directory tree
+		if (!PluginUtils::ValidateNewPluginFolder(FPluginUtils::GetPluginFolder(PluginLocation, PluginName, true /*bFullPath*/), FailReason))
 		{
-			FString ExistingFilePath = FPluginUtils::GetPluginFolder(PluginLocation, PluginName, true /*bFullPath*/);
-			while (!ExistingFilePath.IsEmpty())
-			{
-				if (FPaths::FileExists(ExistingFilePath))
-				{
-					break;
-				}
-				ExistingFilePath = FPaths::GetPath(ExistingFilePath);
-			}
-			
-			if (!ExistingFilePath.IsEmpty())
-			{
-				if (FailReason)
-				{
-					*FailReason = FText::Format(LOCTEXT("PluginLocationIsFile", "Plugin location is invalid because the following file is in the way\n{0}"), FText::FromString(ExistingFilePath));
-				}
-				return false;
-			}
+			return false;
 		}
 	}
 
-	// Check to see if a discovered plugin with this name exists (at any path)
-	if (TSharedPtr<IPlugin> ExistingPlugin = IPluginManager::Get().FindPlugin(PluginName))
+	if (!PluginUtils::ValidateNewPluginName(PluginName, FailReason))
+	{
+		return false;
+	}
+
+	return true;
+}
+
+bool FPluginUtils::ValidateNewPluginFilePath(const FString& PluginFilePath, FText* FailReason /*= nullptr*/)
+{
+	const FString PluginName = GetPluginName(PluginFilePath);
+
+	// Check whether the plugin name is valid
+	if (!IsValidPluginName(PluginName, FailReason))
 	{
-		if (FailReason)
-		{
-			*FailReason = FText::Format(LOCTEXT("PluginNameAlreadyInUse", "Plugin name is already in use by\n{0}"), FText::FromString(FPaths::ConvertRelativePathToFull(ExistingPlugin->GetDescriptorFileName())));
-		}
 		return false;
 	}
 
-	// Check if a mount point with the same name exists
+	if (!PluginUtils::ValidateNewPluginFilePath(PluginFilePath, FailReason))
 	{
-		FString MountPoint;
-		MountPoint.Reserve(PluginName.Len() + 2);
-		MountPoint.AppendChar(TEXT('/'));
-		MountPoint.Append(PluginName);
-		MountPoint.AppendChar(TEXT('/'));
-		if (FPackageName::MountPointExists(MountPoint))
-		{
-			if (FailReason)
-			{
-				*FailReason = FText::Format(LOCTEXT("MountPointExists", "Mount point {0} already exists"), FText::FromString(MountPoint));
-			}
-			return false;
-		}
+		return false;
+	}
+
+	if (!PluginUtils::ValidateNewPluginFolder(FPluginUtils::GetPluginFolder(PluginFilePath, true /*bFullPath*/), FailReason))
+	{
+		return false;
+	}
+
+	if (!PluginUtils::ValidateNewPluginName(PluginName, FailReason))
+	{
+		return false;
 	}
 
 	return true;
diff --git a/Engine/Plugins/Developer/PluginUtils/Source/PluginUtils/Public/PluginUtils.h b/Engine/Plugins/Developer/PluginUtils/Source/PluginUtils/Public/PluginUtils.h
--- a/Engine/Plugins/Developer/PluginUtils/Source/PluginUtils/Public/PluginUtils.h
+++ b/Engine/Plugins/Developer/PluginUtils/Source/PluginUtils/Public/PluginUtils.h
@@ -20,6 +20,12 @@ public:
 	 */
 	static FString GetPluginFolder(const FString& PluginLocation, const FString& PluginName, bool bFullPath = false);
 
+	/**
+	 * Returns the plugin folder from its uplugin file path.
+	 * @param PluginFilePath Plugin descriptor file path
+	 */
+	static FString GetPluginFolder(const FString& PluginFilePath, bool bFullPath = false);
+
 	/** 
 	 * Returns the uplugin file path.
 	 * @param PluginLocation directory that contains the plugin folder
@@ -28,6 +34,12 @@ public:
 	 */
 	static FString GetPluginFilePath(const FString& PluginLocation, const FString& PluginName, bool bFullPath = false);
 
+	/**
+	 * Returns the plugin name from its uplugin file path.
+	 * @param PluginFilePath Plugin descriptor file path
+	 */
+	static FString GetPluginName(const FString& PluginFilePath);
+
 	/**
 	 * Returns the plugin Content folder.
 	 * @param PluginLocation directory that contains the plugin folder
@@ -36,6 +48,12 @@ public:
 	 */
 	static FString GetPluginContentFolder(const FString& PluginLocation, const FString& PluginName, bool bFullPath = false);
 
+	/**
+	 * Returns the plugin Content folder from its uplugin file path.
+	 * @param PluginFilePath Plugin descriptor file path
+	 */
+	static FString GetPluginContentFolder(const FString& PluginFilePath, bool bFullPath = false);
+
 	/**
 	 * Returns the plugin Resources folder.
 	 * @param PluginLocation directory that contains the plugin folder
@@ -44,6 +62,12 @@ public:
 	 */
 	static FString GetPluginResourcesFolder(const FString& PluginLocation, const FString& PluginName, bool bFullPath = false);
 
+	/**
+	 * Returns the plugin Resources folder from its uplugin file path.
+	 * @param PluginFilePath Plugin descriptor file path
+	 */
+	static FString GetPluginResourcesFolder(const FString& PluginFilePath, bool bFullPath = false);
+
 	/**
 	 * Parameters for creating a new plugin.
 	 */
@@ -127,9 +151,15 @@ public:
 
 		/**
 		 * Folders containing template files to copy into the plugin folder (Required if Descriptor.Modules is not empty).
-		 * Occurrences of the string PLUGIN_NAME in the filename or file content will be replaced by the plugin name.
+		 * Occurrences of the string PLUGIN_NAME in the filename or file content will be replaced by NameToReplace.
 		 */
 		TArray<FString> TemplateFolders;
+
+		/**
+		 * String used to replace PLUGIN_NAME in templates.
+		 * If empty, the plugin name will be used.
+		 */
+		FString NameToReplace;
 	};
 
 	/**
@@ -213,8 +243,33 @@ public:
 	 * @note MountParams.OutFailReason outputs the reason the plugin creation or loading failed (if applicable)
 	 * @note Will fail if the plugin already exists
 	 */
+	UE_DEPRECATED(5.6, "Please set NameToReplace in the FNewPluginParamsWithDescriptor and call the CreateAndLoadNewPlugin overload without the NameToReplace parameter.")
 	static TSharedPtr<IPlugin> CreateAndLoadNewPlugin(const FString& PluginName, const FString& NameToReplace, const FString& PluginLocation, const FNewPluginParamsWithDescriptor& CreationParams, FLoadPluginParams& LoadParams);
 	
+	/**
+	 * Helper to create and load a new plugin
+	 * @param PluginFilePath Plugin descriptor file path
+	 * @param CreationParams Plugin creation parameters
+	 * @param MountParams Plugin loading parameters
+	 * @return The newly created plugin. If something goes wrong during the creation process, the plugin folder gets deleted and null is returned.
+	 * @note MountParams.OutFailReason outputs the reason the plugin creation or loading failed (if applicable)
+	 * @note Will fail if the plugin already exists
+	 * @note This function can circumvent the standard plugin directory structure, in which case the newly created plugin may also need to be loaded by full path
+	 */
+	static TSharedPtr<IPlugin> CreateAndLoadNewPlugin(const FString& PluginFilePath, const FNewPluginParams& CreationParams, FLoadPluginParams& LoadParams);
+
+	/**
+	 * Helper to create and load a new plugin
+	 * @param PluginFilePath Plugin descriptor file path
+	 * @param CreationParams Plugin creation parameters
+	 * @param LoadParams Plugin loading parameters
+	 * @return The newly created plugin. If something goes wrong during the creation process, the plugin folder gets deleted and null is returned.
+	 * @note MountParams.OutFailReason outputs the reason the plugin creation or loading failed (if applicable)
+	 * @note Will fail if the plugin already exists
+	 * @note This function can circumvent the standard plugin directory structure, in which case the newly created plugin may also need to be loaded by full path
+	 */
+	static TSharedPtr<IPlugin> CreateAndLoadNewPlugin(const FString& PluginFilePath, const FNewPluginParamsWithDescriptor& CreationParams, FLoadPluginParams& LoadParams);
+
 	/**
 	 * Helper to create and mount a new plugin.
 	 * @param PluginName Plugin name
@@ -268,21 +323,21 @@ public:
 
 	/**
 	 * Load/mount the specified plugin
-	 * @param PluginFileName Plugin descriptor file path
+	 * @param PluginFilePath Plugin descriptor file path
 	 * @param LoadParams Plugin loading parameters
 	 * @return The loaded plugin or null on failure
 	 */
-	static TSharedPtr<IPlugin> LoadPlugin(const FString& PluginFileName, FLoadPluginParams& LoadParams);
+	static TSharedPtr<IPlugin> LoadPlugin(const FString& PluginFilePath, FLoadPluginParams& LoadParams);
 
 	/**
 	 * Load/mount the specified plugin
-	 * @param PluginFileName Plugin descriptor file path
+	 * @param PluginFilePath Plugin descriptor file path
 	 * @return The loaded plugin or null on failure
 	 */
-	static TSharedPtr<IPlugin> LoadPlugin(const FString& PluginFileName)
+	static TSharedPtr<IPlugin> LoadPlugin(const FString& PluginFilePath)
 	{
 		FLoadPluginParams Params;
-		return LoadPlugin(PluginFileName, Params);
+		return LoadPlugin(PluginFilePath, Params);
 	}
 
 	/**
@@ -417,6 +472,15 @@ public:
 	 */
 	static bool ValidateNewPluginNameAndLocation(const FString& PluginName, const FString& PluginLocation = FString(), FText* FailReason = nullptr);
 
+	/**
+	 * Validate that the plugin name is valid, that the name isn't already used by a registered plugin
+	 * and optionally that there isn't an unregistered plugin with that name that exists at the plugin folder
+	 * @param PluginFilePath Plugin descriptor file path
+	 * @param FailReason Optional output text describing why the validation failed
+	 * @return
+	 */
+	static bool ValidateNewPluginFilePath(const FString& PluginFilePath, FText* FailReason = nullptr);
+
 	/**
 	 * Returns whether the specified plugin name is valid, regardless of whether it's already used
 	 * @param PluginName Plugin name
diff --git a/Engine/Plugins/Editor/PluginBrowser/Source/PluginBrowser/Private/SNewPluginWizard.cpp b/Engine/Plugins/Editor/PluginBrowser/Source/PluginBrowser/Private/SNewPluginWizard.cpp
--- a/Engine/Plugins/Editor/PluginBrowser/Source/PluginBrowser/Private/SNewPluginWizard.cpp
+++ b/Engine/Plugins/Editor/PluginBrowser/Source/PluginBrowser/Private/SNewPluginWizard.cpp
@@ -659,6 +659,7 @@ FReply SNewPluginWizard::OnCreatePluginClicked()
 	
 	FPluginUtils::FNewPluginParamsWithDescriptor CreationParams;
 	CreationParams.TemplateFolders = PluginWizardDefinition->GetFoldersForSelection();
+	CreationParams.NameToReplace = GetCurrentPluginName().ToString();
 	CreationParams.Descriptor.bCanContainContent = Template->bCanContainContent;
 
 	if (bHasModules)
@@ -690,7 +691,7 @@ FReply SNewPluginWizard::OnCreatePluginClicked()
 	Template->CustomizeDescriptorBeforeCreation(CreationParams.Descriptor);
 
 
-	TSharedPtr<IPlugin> NewPlugin = FPluginUtils::CreateAndLoadNewPlugin(PluginName, GetCurrentPluginName().ToString(), PluginFolderPath, CreationParams, LoadParams);
+	TSharedPtr<IPlugin> NewPlugin = FPluginUtils::CreateAndLoadNewPlugin(PluginName, PluginFolderPath, CreationParams, LoadParams);
 	const bool bSucceeded = NewPlugin.IsValid();
 
 
