Updated headers using UnrealCodeFixup to make sure dllstorage is on methods/staticvars instead of types. Part 2/n

diff --git a/Engine/Plugins/Developer/PluginUtils/Source/PluginUtils/Public/PluginUtils.h b/Engine/Plugins/Developer/PluginUtils/Source/PluginUtils/Public/PluginUtils.h
--- a/Engine/Plugins/Developer/PluginUtils/Source/PluginUtils/Public/PluginUtils.h
+++ b/Engine/Plugins/Developer/PluginUtils/Source/PluginUtils/Public/PluginUtils.h
@@ -9,7 +9,7 @@
 
 class IPlugin;
 
-class PLUGINUTILS_API FPluginUtils
+class FPluginUtils
 {
 public:
 	/**
@@ -18,13 +18,13 @@ public:
 	 * @param PluginName name of the plugin
 	 * @param bFullPath ensures a full path is returned
 	 */
-	static FString GetPluginFolder(const FString& PluginLocation, const FString& PluginName, bool bFullPath = false);
+	static PLUGINUTILS_API FString GetPluginFolder(const FString& PluginLocation, const FString& PluginName, bool bFullPath = false);
 
 	/**
 	 * Returns the plugin folder from its uplugin file path.
 	 * @param PluginFilePath Plugin descriptor file path
 	 */
-	static FString GetPluginFolder(const FString& PluginFilePath, bool bFullPath = false);
+	static PLUGINUTILS_API FString GetPluginFolder(const FString& PluginFilePath, bool bFullPath = false);
 
 	/** 
 	 * Returns the uplugin file path.
@@ -32,13 +32,13 @@ public:
 	 * @param PluginName name of the plugin
 	 * @param bFullPath ensures a full path is returned
 	 */
-	static FString GetPluginFilePath(const FString& PluginLocation, const FString& PluginName, bool bFullPath = false);
+	static PLUGINUTILS_API FString GetPluginFilePath(const FString& PluginLocation, const FString& PluginName, bool bFullPath = false);
 
 	/**
 	 * Returns the plugin name from its uplugin file path.
 	 * @param PluginFilePath Plugin descriptor file path
 	 */
-	static FString GetPluginName(const FString& PluginFilePath);
+	static PLUGINUTILS_API FString GetPluginName(const FString& PluginFilePath);
 
 	/**
 	 * Returns the plugin Content folder.
@@ -46,13 +46,13 @@ public:
 	 * @param PluginName name of the plugin
 	 * @param bFullPath ensures a full path is returned
 	 */
-	static FString GetPluginContentFolder(const FString& PluginLocation, const FString& PluginName, bool bFullPath = false);
+	static PLUGINUTILS_API FString GetPluginContentFolder(const FString& PluginLocation, const FString& PluginName, bool bFullPath = false);
 
 	/**
 	 * Returns the plugin Content folder from its uplugin file path.
 	 * @param PluginFilePath Plugin descriptor file path
 	 */
-	static FString GetPluginContentFolder(const FString& PluginFilePath, bool bFullPath = false);
+	static PLUGINUTILS_API FString GetPluginContentFolder(const FString& PluginFilePath, bool bFullPath = false);
 
 	/**
 	 * Returns the plugin Resources folder.
@@ -60,13 +60,13 @@ public:
 	 * @param PluginName name of the plugin
 	 * @param bFullPath ensures a full path is returned
 	 */
-	static FString GetPluginResourcesFolder(const FString& PluginLocation, const FString& PluginName, bool bFullPath = false);
+	static PLUGINUTILS_API FString GetPluginResourcesFolder(const FString& PluginLocation, const FString& PluginName, bool bFullPath = false);
 
 	/**
 	 * Returns the plugin Resources folder from its uplugin file path.
 	 * @param PluginFilePath Plugin descriptor file path
 	 */
-	static FString GetPluginResourcesFolder(const FString& PluginFilePath, bool bFullPath = false);
+	static PLUGINUTILS_API FString GetPluginResourcesFolder(const FString& PluginFilePath, bool bFullPath = false);
 
 	/**
 	 * Parameters for creating a new plugin.
@@ -218,7 +218,7 @@ public:
 	 * @note MountParams.OutFailReason outputs the reason the plugin creation or loading failed (if applicable)
 	 * @note Will fail if the plugin already exists
 	 */
-	static TSharedPtr<IPlugin> CreateAndLoadNewPlugin(const FString& PluginName, const FString& PluginLocation, const FNewPluginParams& CreationParams, FLoadPluginParams& LoadParams);
+	static PLUGINUTILS_API TSharedPtr<IPlugin> CreateAndLoadNewPlugin(const FString& PluginName, const FString& PluginLocation, const FNewPluginParams& CreationParams, FLoadPluginParams& LoadParams);
 
 	/**
 	 * Helper to create and load a new plugin
@@ -230,7 +230,7 @@ public:
 	 * @note MountParams.OutFailReason outputs the reason the plugin creation or loading failed (if applicable)
 	 * @note Will fail if the plugin already exists
 	 */
-	static TSharedPtr<IPlugin> CreateAndLoadNewPlugin(const FString& PluginName, const FString& PluginLocation, const FNewPluginParamsWithDescriptor& CreationParams, FLoadPluginParams& LoadParams);
+	static PLUGINUTILS_API TSharedPtr<IPlugin> CreateAndLoadNewPlugin(const FString& PluginName, const FString& PluginLocation, const FNewPluginParamsWithDescriptor& CreationParams, FLoadPluginParams& LoadParams);
 	
 	/**
 	 * Helper to create and load a new plugin
@@ -244,7 +244,7 @@ public:
 	 * @note Will fail if the plugin already exists
 	 */
 	UE_DEPRECATED(5.6, "Please set NameToReplace in the FNewPluginParamsWithDescriptor and call the CreateAndLoadNewPlugin overload without the NameToReplace parameter.")
-	static TSharedPtr<IPlugin> CreateAndLoadNewPlugin(const FString& PluginName, const FString& NameToReplace, const FString& PluginLocation, const FNewPluginParamsWithDescriptor& CreationParams, FLoadPluginParams& LoadParams);
+	static PLUGINUTILS_API TSharedPtr<IPlugin> CreateAndLoadNewPlugin(const FString& PluginName, const FString& NameToReplace, const FString& PluginLocation, const FNewPluginParamsWithDescriptor& CreationParams, FLoadPluginParams& LoadParams);
 	
 	/**
 	 * Helper to create and load a new plugin
@@ -256,7 +256,7 @@ public:
 	 * @note Will fail if the plugin already exists
 	 * @note This function can circumvent the standard plugin directory structure, in which case the newly created plugin may also need to be loaded by full path
 	 */
-	static TSharedPtr<IPlugin> CreateAndLoadNewPlugin(const FString& PluginFilePath, const FNewPluginParams& CreationParams, FLoadPluginParams& LoadParams);
+	static PLUGINUTILS_API TSharedPtr<IPlugin> CreateAndLoadNewPlugin(const FString& PluginFilePath, const FNewPluginParams& CreationParams, FLoadPluginParams& LoadParams);
 
 	/**
 	 * Helper to create and load a new plugin
@@ -268,7 +268,7 @@ public:
 	 * @note Will fail if the plugin already exists
 	 * @note This function can circumvent the standard plugin directory structure, in which case the newly created plugin may also need to be loaded by full path
 	 */
-	static TSharedPtr<IPlugin> CreateAndLoadNewPlugin(const FString& PluginFilePath, const FNewPluginParamsWithDescriptor& CreationParams, FLoadPluginParams& LoadParams);
+	static PLUGINUTILS_API TSharedPtr<IPlugin> CreateAndLoadNewPlugin(const FString& PluginFilePath, const FNewPluginParamsWithDescriptor& CreationParams, FLoadPluginParams& LoadParams);
 
 	/**
 	 * Helper to create and mount a new plugin.
@@ -282,7 +282,7 @@ public:
 	 */
 	PRAGMA_DISABLE_DEPRECATION_WARNINGS
 	UE_DEPRECATED(5.0, "CreateAndMountNewPlugin is deprecated; please use CreateAndLoadNewPlugin instead")
-	static TSharedPtr<IPlugin> CreateAndMountNewPlugin(const FString& PluginName, const FString& PluginLocation, const FNewPluginParams& CreationParams, const FMountPluginParams& MountParams, FText& FailReason);
+	static PLUGINUTILS_API TSharedPtr<IPlugin> CreateAndMountNewPlugin(const FString& PluginName, const FString& PluginLocation, const FNewPluginParams& CreationParams, const FMountPluginParams& MountParams, FText& FailReason);
 	PRAGMA_ENABLE_DEPRECATION_WARNINGS
 
 	/**
@@ -297,8 +297,8 @@ public:
 	 */
 	PRAGMA_DISABLE_DEPRECATION_WARNINGS
 	UE_DEPRECATED(5.0, "CreateAndMountNewPlugin is deprecated; please use CreateAndLoadNewPlugin instead")
-	static TSharedPtr<IPlugin> CreateAndMountNewPlugin(const FString& PluginName, const FString& PluginLocation, const FNewPluginParamsWithDescriptor& CreationParams, const FMountPluginParams& MountParams, FText& FailReason);
-	PRAGMA_ENABLE_DEPRECATION_WARNINGS
+	static PLUGINUTILS_API TSharedPtr<IPlugin> CreateAndMountNewPlugin(const FString& PluginName, const FString& PluginLocation, const FNewPluginParamsWithDescriptor& CreationParams, const FMountPluginParams& MountParams, FText& FailReason);
+	PLUGINUTILS_API PRAGMA_ENABLE_DEPRECATION_WARNINGS
 
 	/**
 	 * Load/mount the specified plugin
@@ -327,7 +327,7 @@ public:
 	 * @param LoadParams Plugin loading parameters
 	 * @return The loaded plugin or null on failure
 	 */
-	static TSharedPtr<IPlugin> LoadPlugin(const FString& PluginFilePath, FLoadPluginParams& LoadParams);
+	static PLUGINUTILS_API TSharedPtr<IPlugin> LoadPlugin(const FString& PluginFilePath, FLoadPluginParams& LoadParams);
 
 	/**
 	 * Load/mount the specified plugin
@@ -350,8 +350,8 @@ public:
 	 */
 	PRAGMA_DISABLE_DEPRECATION_WARNINGS
 	UE_DEPRECATED(5.0, "MountPlugin is deprecated; please use LoadPlugin instead")
-	static TSharedPtr<IPlugin> MountPlugin(const FString& PluginName, const FString& PluginLocation, const FMountPluginParams& MountParams, FText& FailReason);
-	PRAGMA_ENABLE_DEPRECATION_WARNINGS
+	static PLUGINUTILS_API TSharedPtr<IPlugin> MountPlugin(const FString& PluginName, const FString& PluginLocation, const FMountPluginParams& MountParams, FText& FailReason);
+	PLUGINUTILS_API PRAGMA_ENABLE_DEPRECATION_WARNINGS
 
 	/**
 	 * Finds a loaded plugin from a plugin descriptor file path
@@ -364,8 +364,8 @@ public:
 	 * @param PackagePath Unreal path of the package
 	 * @return Plugin or null if the package is not in a plugin
 	 */
-	static TSharedPtr<IPlugin> FindPluginFromPackagePath(FName PackagePath);
-	static TSharedPtr<IPlugin> FindPluginFromPackagePath(FStringView PackagePath);
+	static PLUGINUTILS_API TSharedPtr<IPlugin> FindPluginFromPackagePath(FName PackagePath);
+	static PLUGINUTILS_API TSharedPtr<IPlugin> FindPluginFromPackagePath(FStringView PackagePath);
 
 	/**
 	 * Unload assets from the specified plugin and unmount it
@@ -375,7 +375,7 @@ public:
 	 * @param OutFailReason Outputs the reason of the failure if any
 	 * @return Whether the plugin was successfully unloaded
 	 */
-	static bool UnloadPlugin(const TSharedRef<IPlugin>& Plugin, FText* OutFailReason = nullptr);
+	static PLUGINUTILS_API bool UnloadPlugin(const TSharedRef<IPlugin>& Plugin, FText* OutFailReason = nullptr);
 
 	/**
 	 * Unload assets from the specified plugin and unmount it
@@ -385,7 +385,7 @@ public:
 	 * @param OutFailReason Outputs the reason of the failure if any
 	 * @return Whether the plugin was successfully unloaded
 	 */
-	static bool UnloadPlugin(const FString& PluginName, FText* OutFailReason = nullptr);
+	static PLUGINUTILS_API bool UnloadPlugin(const FString& PluginName, FText* OutFailReason = nullptr);
 
 	/**
 	 * Unload assets from the specified plugins and unmount them
@@ -395,7 +395,7 @@ public:
 	 * @param OutFailReason Outputs the reason of the failure if any
 	 * @return Whether all plugins were successfully unloaded
 	 */
-	static bool UnloadPlugins(const TConstArrayView<TSharedRef<IPlugin>> Plugins, FText* OutFailReason = nullptr);
+	static PLUGINUTILS_API bool UnloadPlugins(const TConstArrayView<TSharedRef<IPlugin>> Plugins, FText* OutFailReason = nullptr);
 
 	/**
 	 * Unload assets from the specified plugins and unmount them
@@ -405,7 +405,7 @@ public:
 	 * @param OutFailReason Outputs the reason of the failure if any
 	 * @return Whether all plugins were successfully unloaded
 	 */
-	static bool UnloadPlugins(const TConstArrayView<FString> PluginNames, FText* OutFailReason = nullptr);
+	static PLUGINUTILS_API bool UnloadPlugins(const TConstArrayView<FString> PluginNames, FText* OutFailReason = nullptr);
 
 	/**
 	 * Unload assets from the specified plugin but does not unmount it
@@ -414,7 +414,7 @@ public:
 	 * @param OutFailReason Outputs the reason of the failure if any
 	 * @return Whether plugin assets were successfully unloaded
 	 */
-	static bool UnloadPluginAssets(const TSharedRef<IPlugin>& Plugin, FText* OutFailReason = nullptr);
+	static PLUGINUTILS_API bool UnloadPluginAssets(const TSharedRef<IPlugin>& Plugin, FText* OutFailReason = nullptr);
 
 	/**
 	 * Unload assets from the specified plugin but does not unmount it
@@ -423,7 +423,7 @@ public:
 	 * @param OutFailReason Outputs the reason of the failure if any
 	 * @return Whether plugin assets were successfully unloaded
 	 */
-	static bool UnloadPluginAssets(const FString& PluginName, FText* OutFailReason = nullptr);
+	static PLUGINUTILS_API bool UnloadPluginAssets(const FString& PluginName, FText* OutFailReason = nullptr);
 
 	/**
 	 * Unload assets from the specified plugins but does not unmount them
@@ -432,7 +432,7 @@ public:
 	 * @param OutFailReason Outputs the reason of the failure if any
 	 * @return Whether plugin assets were successfully unloaded
 	 */
-	static bool UnloadPluginsAssets(const TConstArrayView<TSharedRef<IPlugin>> Plugins, FText* OutFailReason = nullptr);
+	static PLUGINUTILS_API bool UnloadPluginsAssets(const TConstArrayView<TSharedRef<IPlugin>> Plugins, FText* OutFailReason = nullptr);
 
 	/**
 	 * Unload assets from the specified plugin but does not unmount them
@@ -441,7 +441,7 @@ public:
 	 * @param OutFailReason Outputs the reason of the failure if any
 	 * @return Whether plugin assets were successfully unloaded
 	 */
-	static bool UnloadPluginsAssets(const TSet<FString>& PluginNames, FText* OutFailReason = nullptr);
+	static PLUGINUTILS_API bool UnloadPluginsAssets(const TSet<FString>& PluginNames, FText* OutFailReason = nullptr);
 
 	/**
 	 * Unload assets from the specified plugin but does not unmount them
@@ -450,7 +450,7 @@ public:
 	 * @param OutFailReason Outputs the reason of the failure if any
 	 * @return Whether plugin assets were successfully unloaded
 	 */
-	static bool UnloadPluginsAssets(const TConstArrayView<FString> PluginNames, FText* OutFailReason = nullptr);
+	static PLUGINUTILS_API bool UnloadPluginsAssets(const TConstArrayView<FString> PluginNames, FText* OutFailReason = nullptr);
 
 	/**
 	 * Adds a directory to the list of paths that are recursively searched for plugins, 
@@ -460,7 +460,7 @@ public:
 	 * @param bUpdateProjectFile Whether to update the project additional plugin directories (persistently saved in uproject file) if needed
 	 * @return Whether the plugin search path was modified
 	 */
-	static bool AddToPluginSearchPathIfNeeded(const FString& Dir, bool bRefreshPlugins = false, bool bUpdateProjectFile = false);
+	static PLUGINUTILS_API bool AddToPluginSearchPathIfNeeded(const FString& Dir, bool bRefreshPlugins = false, bool bUpdateProjectFile = false);
 
 	/**
 	 * Validate that the plugin name is valid, that the name isn't already used by a registered plugin
@@ -470,7 +470,7 @@ public:
 	 * @param FailReason Optional output text describing why the validation failed
 	 * @return
 	 */
-	static bool ValidateNewPluginNameAndLocation(const FString& PluginName, const FString& PluginLocation = FString(), FText* FailReason = nullptr);
+	static PLUGINUTILS_API bool ValidateNewPluginNameAndLocation(const FString& PluginName, const FString& PluginLocation = FString(), FText* FailReason = nullptr);
 
 	/**
 	 * Validate that the plugin name is valid, that the name isn't already used by a registered plugin
@@ -479,7 +479,7 @@ public:
 	 * @param FailReason Optional output text describing why the validation failed
 	 * @return
 	 */
-	static bool ValidateNewPluginFilePath(const FString& PluginFilePath, FText* FailReason = nullptr);
+	static PLUGINUTILS_API bool ValidateNewPluginFilePath(const FString& PluginFilePath, FText* FailReason = nullptr);
 
 	/**
 	 * Returns whether the specified plugin name is valid, regardless of whether it's already used
@@ -487,7 +487,7 @@ public:
 	 * @param FailReason Optional output text specifying what is wrong with the plugin name
 	 * @param PluginTermReplacement If set, replaces the term "plugin" in the fail reason message
 	 */
-	static bool IsValidPluginName(const FString& PluginName, FText* FailReason = nullptr, const FText* PluginTermReplacement = nullptr);
+	static PLUGINUTILS_API bool IsValidPluginName(const FString& PluginName, FText* FailReason = nullptr, const FText* PluginTermReplacement = nullptr);
 };
 
 #endif //if WITH_EDITOR
diff --git a/Engine/Plugins/Developer/ProjectLauncher/Source/ProjectLauncher/Public/Extension/LaunchExtension.h b/Engine/Plugins/Developer/ProjectLauncher/Source/ProjectLauncher/Public/Extension/LaunchExtension.h
--- a/Engine/Plugins/Developer/ProjectLauncher/Source/ProjectLauncher/Public/Extension/LaunchExtension.h
+++ b/Engine/Plugins/Developer/ProjectLauncher/Source/ProjectLauncher/Public/Extension/LaunchExtension.h
@@ -7,6 +7,8 @@
 #include "ILauncherProfile.h"
 #include "Framework/MultiBox/MultiBoxBuilder.h"
 
+#define UE_API PROJECTLAUNCHER_API
+
 namespace ProjectLauncher
 {
 	class FLaunchExtension;
@@ -22,7 +24,7 @@ namespace ProjectLauncher
 	 *     return MakeShared<FMyLaunchExtensionInstance>(InArgs);
 	 * }
 	 */
-	class PROJECTLAUNCHER_API FLaunchExtensionInstance : public TSharedFromThis<FLaunchExtensionInstance>
+	class FLaunchExtensionInstance : public TSharedFromThis<FLaunchExtensionInstance>
 	{
 	public:
 		struct FArgs
@@ -33,28 +35,28 @@ namespace ProjectLauncher
 			TSharedRef<FLaunchExtension> Extension;
 		};
 
-		FLaunchExtensionInstance( FArgs& InArgs );
-		virtual ~FLaunchExtensionInstance();
+		UE_API FLaunchExtensionInstance( FArgs& InArgs );
+		UE_API virtual ~FLaunchExtensionInstance();
 
 		/**
 		 * Returns the parameters that this extension provides. They will be added to the submenu.
 		 */
-		virtual bool GetExtensionParameters( TArray<FString>& OutParameters ) const;
+		UE_API virtual bool GetExtensionParameters( TArray<FString>& OutParameters ) const;
 
 		/**
 		 * Returns the user-facing name for the given parameter. It will default to the parameter itself.
 		 */
-		virtual FText GetExtensionParameterDisplayName( const FString& InParameter ) const;
+		UE_API virtual FText GetExtensionParameterDisplayName( const FString& InParameter ) const;
 
 		/**
 		 * Returns the user facing variables that this extension provides, in "$(name)" format
 		 */
-		virtual bool GetExtensionVariables( TArray<FString>& OutVariables ) const;
+		UE_API virtual bool GetExtensionVariables( TArray<FString>& OutVariables ) const;
 
 		/**
 		 * Returns the current value for the given variable
 		 */
-		virtual bool GetExtensionVariableValue( const FString& InParameter, FString& OutValue ) const;
+		UE_API virtual bool GetExtensionVariableValue( const FString& InParameter, FString& OutValue ) const;
 
 		/**
 		 * Hook to allow the extension to extend the extension parameters menu
@@ -97,7 +99,7 @@ namespace ProjectLauncher
 		 * @param InParameter the parameter to test, e.g. -key=value or -param
 		 * @returns whether the parameter is being used
 		 */
-		bool IsParameterUsed( const FString& InParameter ) const;
+		UE_API bool IsParameterUsed( const FString& InParameter ) const;
 
 		/** 
 		 * Add or remove the given parameter on the command line
@@ -105,21 +107,21 @@ namespace ProjectLauncher
 		 * @param InParameter the parameter to add or remove, e.g. -key=value or -param
 		 * @param bUsed whether to add or remove the parameter
 		 */
-		void SetParameterUsed( const FString& InParameter, bool bUsed );
+		UE_API void SetParameterUsed( const FString& InParameter, bool bUsed );
 
 		/** 
 		 * Add the given parameter to the command line
 		 * 
 		 * @param InParameter the parameter to add, e.g. -key=value or -param
 		 */
-		void AddParameter( const FString& InParameter );
+		UE_API void AddParameter( const FString& InParameter );
 
 		/** 
 		 * Remove the given parameter from the command line
 		 * 
 		 * @param InParameter the parameter to remove, e.g. -key=value or -param
 		 */
-		void RemoveParameter( const FString& InParameter );
+		UE_API void RemoveParameter( const FString& InParameter );
 
 		/** 
 		 * Retrieve the current value for the given command line parameter
@@ -127,7 +129,7 @@ namespace ProjectLauncher
 		 * @param InParameter the parameter to query, e.g. -key=value
 		 * @returns the value of the parameter
 		 */
-		FString GetParameterValue( const FString& InParameter ) const;
+		UE_API FString GetParameterValue( const FString& InParameter ) const;
 
 		/** 
 		 * Update the value of the given command line parameter
@@ -136,7 +138,7 @@ namespace ProjectLauncher
 		 * @param InNewValue the new value to use
 		 * @returns true unless InParameter is not a key/value property
 		 */
-		bool UpdateParameterValue( const FString& InParameter, const FString& InNewValue );
+		UE_API bool UpdateParameterValue( const FString& InParameter, const FString& InNewValue );
 
 		/** 
 		 * Get the current state of the given command line parameter, allowing for user-changed values
@@ -144,7 +146,7 @@ namespace ProjectLauncher
 		 * @param InParameter the parameter to test, e.g. -key=value
 		 * @returns The parameter as it is now, allowing for user modifications, e.g. -key=value or -key=some_new_value
 		 */
-		FString GetFinalParameter( const FString& InParameter ) const;
+		UE_API FString GetFinalParameter( const FString& InParameter ) const;
 
 		/** 
 		 * Get the current command line
@@ -152,7 +154,7 @@ namespace ProjectLauncher
 		 *
 		 * @returns the current command line string
 		 */
-		FString GetCommandLine() const;
+		UE_API FString GetCommandLine() const;
 
 		/** 
 		 * Update the current command line
@@ -160,7 +162,7 @@ namespace ProjectLauncher
 		 * 
 		 * @param CommandLine the new command line
 		 */
-		void SetCommandLine( const FString& CommandLine );
+		UE_API void SetCommandLine( const FString& CommandLine );
 
 		/** 
 		 * Get the profile we were instantiated for
@@ -195,7 +197,7 @@ namespace ProjectLauncher
 		 * @param DefaultValue value to return if the value isn't found
 		 * @returns the value of the string, or DefaultValue if it isn't found
 		 */
-		FString GetConfigString( EConfig Config, const TCHAR* Name, const TCHAR* DefaultValue = TEXT("") ) const;
+		UE_API FString GetConfigString( EConfig Config, const TCHAR* Name, const TCHAR* DefaultValue = TEXT("") ) const;
 
 		/**
 		 * Read a configuration bool value
@@ -205,7 +207,7 @@ namespace ProjectLauncher
 		 * @param DefaultValue value to return if the value isn't found
 		 * @returns the value of the bool, or DefaultValue if it isn't found
 		 */
-		bool GetConfigBool( EConfig Config, const TCHAR* Name, bool DefaultValue = false ) const;
+		UE_API bool GetConfigBool( EConfig Config, const TCHAR* Name, bool DefaultValue = false ) const;
 
 		/**
 		 * Write a configuration string value
@@ -214,7 +216,7 @@ namespace ProjectLauncher
 		 * @param Name name of the value to write
 		 * @param Value value to write
 		 */
-		void SetConfigString( EConfig Config, const TCHAR* Name, const FString& Value ) const;
+		UE_API void SetConfigString( EConfig Config, const TCHAR* Name, const FString& Value ) const;
 
 		/**
 		 * Write a configuration bool value
@@ -223,7 +225,7 @@ namespace ProjectLauncher
 		 * @param Name name of the value to write
 		 * @param Value value to write
 		 */
-		void SetConfigBool( EConfig Config, const TCHAR* Name, const bool& Value ) const;
+		UE_API void SetConfigBool( EConfig Config, const TCHAR* Name, const bool& Value ) const;
 
 		/**
 		 * Get the final key name to use for reading & writing a configuration value
@@ -233,13 +235,13 @@ namespace ProjectLauncher
 		 * 
 		 * @returns namespaced version of Name
 		 */
-		FString GetConfigKeyName( EConfig Config, const TCHAR* Name ) const;
+		UE_API FString GetConfigKeyName( EConfig Config, const TCHAR* Name ) const;
 
 
 
 	private:
-		bool IsParameterGroup( const FString& InParameter ) const;
-		bool TryRemoveParameterGroup( const FString& InParameter );
+		UE_API bool IsParameterGroup( const FString& InParameter ) const;
+		UE_API bool TryRemoveParameterGroup( const FString& InParameter );
 
 
 		ILauncherProfileRef Profile;
@@ -248,7 +250,7 @@ namespace ProjectLauncher
 		TSharedRef<FLaunchExtension> Extension;
 
 		friend class FLaunchProfileTreeNode;
-		void MakeCommandLineSubmenu( FMenuBuilder& MenuBuilder );
+		UE_API void MakeCommandLineSubmenu( FMenuBuilder& MenuBuilder );
 
 	};
 
@@ -262,7 +264,7 @@ namespace ProjectLauncher
 	 *   TSharedPtr<FMyLaunchExtension> MyExtension = MakeShared<FMyLaunchExtension>();
 	 *   IProjectLauncherModule::Get().RegisterExtension( MyExtension.ToSharedRef() );
 	 */
-	class PROJECTLAUNCHER_API FLaunchExtension : public TSharedFromThis<FLaunchExtension>
+	class FLaunchExtension : public TSharedFromThis<FLaunchExtension>
 	{
 	public:
 		virtual ~FLaunchExtension() = default;
@@ -293,9 +295,10 @@ namespace ProjectLauncher
 		 * @param InModel helper class
 		 * @returns array of all compatible extensions
 		 */
-		static TArray<TSharedPtr<FLaunchExtensionInstance>> CreateExtensionInstancesForProfile( ILauncherProfileRef InProfile, TSharedRef<FModel> InModel, ILaunchProfileTreeBuilder* InTreeBuilder );
+		static UE_API TArray<TSharedPtr<FLaunchExtensionInstance>> CreateExtensionInstancesForProfile( ILauncherProfileRef InProfile, TSharedRef<FModel> InModel, ILaunchProfileTreeBuilder* InTreeBuilder );
 	};
 
 
 };
 
+#undef UE_API
diff --git a/Engine/Plugins/Developer/ProjectLauncher/Source/ProjectLauncher/Public/ProfileTree/BasicProfileTreeBuilder.h b/Engine/Plugins/Developer/ProjectLauncher/Source/ProjectLauncher/Public/ProfileTree/BasicProfileTreeBuilder.h
--- a/Engine/Plugins/Developer/ProjectLauncher/Source/ProjectLauncher/Public/ProfileTree/BasicProfileTreeBuilder.h
+++ b/Engine/Plugins/Developer/ProjectLauncher/Source/ProjectLauncher/Public/ProfileTree/BasicProfileTreeBuilder.h
@@ -4,6 +4,8 @@
 
 #include "ProfileTree/GenericProfileTreeBuilder.h"
 
+#define UE_API PROJECTLAUNCHER_API
+
 namespace ProjectLauncher
 {
 
@@ -24,12 +26,14 @@ namespace ProjectLauncher
 
 
 
-	class PROJECTLAUNCHER_API FBasicProfileTreeBuilderFactory : public ILaunchProfileTreeBuilderFactory
+	class FBasicProfileTreeBuilderFactory : public ILaunchProfileTreeBuilderFactory
 	{
 	public:
-		virtual TSharedPtr<ILaunchProfileTreeBuilder> TryCreateTreeBuilder( const ILauncherProfileRef& InProfile, const TSharedRef<FModel>& InModel ) override;
+		UE_API virtual TSharedPtr<ILaunchProfileTreeBuilder> TryCreateTreeBuilder( const ILauncherProfileRef& InProfile, const TSharedRef<FModel>& InModel ) override;
 
 		virtual bool IsProfileTypeSupported(EProfileType ProfileType ) const override { return ProfileType == EProfileType::Basic; }
 
 	};
 }
+
+#undef UE_API
diff --git a/Engine/Plugins/Developer/ProjectLauncher/Source/ProjectLauncher/Public/ProfileTree/CustomProfileTreeBuilder.h b/Engine/Plugins/Developer/ProjectLauncher/Source/ProjectLauncher/Public/ProfileTree/CustomProfileTreeBuilder.h
--- a/Engine/Plugins/Developer/ProjectLauncher/Source/ProjectLauncher/Public/ProfileTree/CustomProfileTreeBuilder.h
+++ b/Engine/Plugins/Developer/ProjectLauncher/Source/ProjectLauncher/Public/ProfileTree/CustomProfileTreeBuilder.h
@@ -4,15 +4,17 @@
 
 #include "ProfileTree/GenericProfileTreeBuilder.h"
 
+#define UE_API PROJECTLAUNCHER_API
+
 namespace ProjectLauncher
 {
-	class PROJECTLAUNCHER_API FCustomProfileTreeBuilder : public FGenericProfileTreeBuilder
+	class FCustomProfileTreeBuilder : public FGenericProfileTreeBuilder
 	{
 	public:
-		FCustomProfileTreeBuilder( const ILauncherProfileRef& Profile, const TSharedRef<FModel>& InModel );
+		UE_API FCustomProfileTreeBuilder( const ILauncherProfileRef& Profile, const TSharedRef<FModel>& InModel );
 		virtual ~FCustomProfileTreeBuilder() = default;
 
-		virtual void Construct() override;
+		UE_API virtual void Construct() override;
 
 		virtual FString GetName() const override
 		{
@@ -21,13 +23,15 @@ namespace ProjectLauncher
 	};
 
 
-	class PROJECTLAUNCHER_API FCustomProfileTreeBuilderFactory : public ILaunchProfileTreeBuilderFactory
+	class FCustomProfileTreeBuilderFactory : public ILaunchProfileTreeBuilderFactory
 	{
 	public:
-		virtual TSharedPtr<ILaunchProfileTreeBuilder> TryCreateTreeBuilder( const ILauncherProfileRef& Profile, const TSharedRef<FModel>& InModel ) override;
+		UE_API virtual TSharedPtr<ILaunchProfileTreeBuilder> TryCreateTreeBuilder( const ILauncherProfileRef& Profile, const TSharedRef<FModel>& InModel ) override;
 
 		virtual bool IsProfileTypeSupported(EProfileType ProfileType ) const override { return ProfileType == EProfileType::Custom; }
 
 	};
 
 }
+
+#undef UE_API
diff --git a/Engine/Plugins/Developer/ProjectLauncher/Source/ProjectLauncher/Public/ProfileTree/GenericProfileTreeBuilder.h b/Engine/Plugins/Developer/ProjectLauncher/Source/ProjectLauncher/Public/ProfileTree/GenericProfileTreeBuilder.h
--- a/Engine/Plugins/Developer/ProjectLauncher/Source/ProjectLauncher/Public/ProfileTree/GenericProfileTreeBuilder.h
+++ b/Engine/Plugins/Developer/ProjectLauncher/Source/ProjectLauncher/Public/ProfileTree/GenericProfileTreeBuilder.h
@@ -5,6 +5,8 @@
 #include "ProfileTree/ILaunchProfileTreeBuilder.h"
 #include "Model/ProjectLauncherModel.h"
 
+#define UE_API PROJECTLAUNCHER_API
+
 class SSearchableComboBox;
 class SCustomLaunchMapListView;
 class SCustomLaunchDeviceListView;
@@ -22,13 +24,13 @@ namespace ProjectLauncher
 	 *        return MakeShared<FMyProfileTreeBuilder>(InProfile, InModel);
 	 *    }
 	 */
-	class PROJECTLAUNCHER_API FGenericProfileTreeBuilder : public ILaunchProfileTreeBuilder, public TSharedFromThis<FGenericProfileTreeBuilder>
+	class FGenericProfileTreeBuilder : public ILaunchProfileTreeBuilder, public TSharedFromThis<FGenericProfileTreeBuilder>
 	{
 	public:
-		FGenericProfileTreeBuilder( const ILauncherProfileRef& Profile, const ILauncherProfileRef& InDefaultProfile, const TSharedRef<FModel>& InModel );
+		UE_API FGenericProfileTreeBuilder( const ILauncherProfileRef& Profile, const ILauncherProfileRef& InDefaultProfile, const TSharedRef<FModel>& InModel );
 		virtual ~FGenericProfileTreeBuilder() = default;
 
-		virtual void Construct() override;
+		UE_API virtual void Construct() override;
 
 		virtual FLaunchProfileTreeDataRef GetProfileTree() override
 		{
@@ -42,33 +44,33 @@ namespace ProjectLauncher
 
 	protected:
 		// default property creation functions
-		void AddProjectProperty( FLaunchProfileTreeNode& HeadingNode );
-		void AddTargetProperty( FLaunchProfileTreeNode& HeadingNode );
-		void AddPlatformProperty( FLaunchProfileTreeNode& HeadingNode );
-		void AddConfigurationProperty( FLaunchProfileTreeNode& HeadingNode );
-		void AddContentSchemeProperty( FLaunchProfileTreeNode& HeadingNode );
-		void AddCompressPakFilesProperty( FLaunchProfileTreeNode& HeadingNode );
-		void AddUseIoStoreProperty( FLaunchProfileTreeNode& HeadingNode );
-		void AddGenerateChunksProperty( FLaunchProfileTreeNode& HeadingNode );
-		void AddImportZenSnapshotProperty( FLaunchProfileTreeNode& HeadingNode );
-		void AddZenPakStreamingPathProperty( FLaunchProfileTreeNode& HeadingNode );
-		void AddIncrementalCookProperty( FLaunchProfileTreeNode& HeadingNode );
-		void AddCookProperty( FLaunchProfileTreeNode& HeadingNode );
-		void AddMapsToCookProperty( FLaunchProfileTreeNode& HeadingNode );
-		void AddAdditionalCookerOptionsProperty( FLaunchProfileTreeNode& HeadingNode );
-		void AddBuildProperty( FLaunchProfileTreeNode& HeadingNode );
-		void AddForceBuildProperty( FLaunchProfileTreeNode& HeadingNode );
-		void AddBuidUATProperty( FLaunchProfileTreeNode& HeadingNode );
-		void AddArchitectureProperty( FLaunchProfileTreeNode& HeadingNode );
-		void AddStagingDirectoryProperty( FLaunchProfileTreeNode& HeadingNode );
-		void AddArchiveBuildProperty( FLaunchProfileTreeNode& HeadingNode );
-		void AddArchiveBuildDirectoryProperty( FLaunchProfileTreeNode& HeadingNode );
-		void AddDeployProperty( FLaunchProfileTreeNode& HeadingNode );
-		void AddIncrementalDeployProperty( FLaunchProfileTreeNode& HeadingNode );
-		void AddTargetDeviceProperty( FLaunchProfileTreeNode& HeadingNode );
-		void AddRunProperty( FLaunchProfileTreeNode& HeadingNode );
-		void AddInitialMapProperty( FLaunchProfileTreeNode& HeadingNode );
-		void AddCommandLineProperty( FLaunchProfileTreeNode& HeadingNode );
+		UE_API void AddProjectProperty( FLaunchProfileTreeNode& HeadingNode );
+		UE_API void AddTargetProperty( FLaunchProfileTreeNode& HeadingNode );
+		UE_API void AddPlatformProperty( FLaunchProfileTreeNode& HeadingNode );
+		UE_API void AddConfigurationProperty( FLaunchProfileTreeNode& HeadingNode );
+		UE_API void AddContentSchemeProperty( FLaunchProfileTreeNode& HeadingNode );
+		UE_API void AddCompressPakFilesProperty( FLaunchProfileTreeNode& HeadingNode );
+		UE_API void AddUseIoStoreProperty( FLaunchProfileTreeNode& HeadingNode );
+		UE_API void AddGenerateChunksProperty( FLaunchProfileTreeNode& HeadingNode );
+		UE_API void AddImportZenSnapshotProperty( FLaunchProfileTreeNode& HeadingNode );
+		UE_API void AddZenPakStreamingPathProperty( FLaunchProfileTreeNode& HeadingNode );
+		UE_API void AddIncrementalCookProperty( FLaunchProfileTreeNode& HeadingNode );
+		UE_API void AddCookProperty( FLaunchProfileTreeNode& HeadingNode );
+		UE_API void AddMapsToCookProperty( FLaunchProfileTreeNode& HeadingNode );
+		UE_API void AddAdditionalCookerOptionsProperty( FLaunchProfileTreeNode& HeadingNode );
+		UE_API void AddBuildProperty( FLaunchProfileTreeNode& HeadingNode );
+		UE_API void AddForceBuildProperty( FLaunchProfileTreeNode& HeadingNode );
+		UE_API void AddBuidUATProperty( FLaunchProfileTreeNode& HeadingNode );
+		UE_API void AddArchitectureProperty( FLaunchProfileTreeNode& HeadingNode );
+		UE_API void AddStagingDirectoryProperty( FLaunchProfileTreeNode& HeadingNode );
+		UE_API void AddArchiveBuildProperty( FLaunchProfileTreeNode& HeadingNode );
+		UE_API void AddArchiveBuildDirectoryProperty( FLaunchProfileTreeNode& HeadingNode );
+		UE_API void AddDeployProperty( FLaunchProfileTreeNode& HeadingNode );
+		UE_API void AddIncrementalDeployProperty( FLaunchProfileTreeNode& HeadingNode );
+		UE_API void AddTargetDeviceProperty( FLaunchProfileTreeNode& HeadingNode );
+		UE_API void AddRunProperty( FLaunchProfileTreeNode& HeadingNode );
+		UE_API void AddInitialMapProperty( FLaunchProfileTreeNode& HeadingNode );
+		UE_API void AddCommandLineProperty( FLaunchProfileTreeNode& HeadingNode );
 
 		// helper callbacks to simplify control enable/visibility
 		FLaunchProfileTreeNode::FGetBool ForPak;
@@ -100,80 +102,80 @@ namespace ProjectLauncher
 		EDeployDeviceOption DeployDeviceOption = EDeployDeviceOption::Default;
 
 
-		FString GetProjectPath() const;
-		void SetProjectName(FString ProjectPath);
-		bool HasProject() const;
+		UE_API FString GetProjectPath() const;
+		UE_API void SetProjectName(FString ProjectPath);
+		UE_API bool HasProject() const;
 
-		FString GetBuildTarget() const;
-		void SetBuildTarget(FString BuildTarget);
-		TArray<EBuildTargetType> GetSupportedBuildTargetTypes() const;
+		UE_API FString GetBuildTarget() const;
+		UE_API void SetBuildTarget(FString BuildTarget);
+		UE_API TArray<EBuildTargetType> GetSupportedBuildTargetTypes() const;
 
-		void SetBuildConfiguration(EBuildConfiguration BuildConfiguration);
-		EBuildConfiguration GetBuildConfiguration() const;
+		UE_API void SetBuildConfiguration(EBuildConfiguration BuildConfiguration);
+		UE_API EBuildConfiguration GetBuildConfiguration() const;
 
-		void SetContentScheme(EContentScheme ContentScheme);
-		bool IsContentSchemeAvailable(EContentScheme, FText& OutReason) const;
+		UE_API void SetContentScheme(EContentScheme ContentScheme);
+		UE_API bool IsContentSchemeAvailable(EContentScheme, FText& OutReason) const;
 
-		FString GetCommandLine() const;
-		void SetCommandLine( const FString& NewCommandLine );
+		UE_API FString GetCommandLine() const;
+		UE_API void SetCommandLine( const FString& NewCommandLine );
 
-		void SetSelectedPlatforms( TArray<FString> SelectedPlatforms );
-		TArray<FString> GetSelectedPlatforms() const;
+		UE_API void SetSelectedPlatforms( TArray<FString> SelectedPlatforms );
+		UE_API TArray<FString> GetSelectedPlatforms() const;
 
-		void SetCook( bool bCook );
-		bool GetCook( ILauncherProfilePtr InProfile = nullptr ) const;
+		UE_API void SetCook( bool bCook );
+		UE_API bool GetCook( ILauncherProfilePtr InProfile = nullptr ) const;
 
-		void SetIncrementalCookMode( ELauncherProfileIncrementalCookMode::Type Mode );
-		ELauncherProfileIncrementalCookMode::Type GetIncrementalCookMode() const;
+		UE_API void SetIncrementalCookMode( ELauncherProfileIncrementalCookMode::Type Mode );
+		UE_API ELauncherProfileIncrementalCookMode::Type GetIncrementalCookMode() const;
 
-		void SetMapsToCook(TArray<FString> MapsToCook);
-		TArray<FString> GetMapsToCook() const;
-		EMapOption GetMapOption() const;
-		void SetMapOption( EMapOption MapOption );
-		float GetMapListHeight() const;
-		void SetMapListHeight( float NewHeight );
-		TSharedRef<SWidget> CreateMapListWidget();
+		UE_API void SetMapsToCook(TArray<FString> MapsToCook);
+		UE_API TArray<FString> GetMapsToCook() const;
+		UE_API EMapOption GetMapOption() const;
+		UE_API void SetMapOption( EMapOption MapOption );
+		UE_API float GetMapListHeight() const;
+		UE_API void SetMapListHeight( float NewHeight );
+		UE_API TSharedRef<SWidget> CreateMapListWidget();
 		float MapListHeight = 300.0f;
 
 
-		void SetDeployDeviceIDs(TArray<FString> DeployDeviceIDs);
-		TArray<FString> GetDeployDeviceIDs() const;
-		EDeployDeviceOption GetDeployDeviceOption() const;
-		void SetDeployDeviceOption( EDeployDeviceOption DeployDeviceOption );
-		float GetDeployDeviceListHeight() const;
-		void SetDeployDeviceListHeight( float NewHeight );
-		void OnDeviceRemoved(FString DeviceID);
-		TSharedRef<SWidget> CreateDeployDeviceWidget();
+		UE_API void SetDeployDeviceIDs(TArray<FString> DeployDeviceIDs);
+		UE_API TArray<FString> GetDeployDeviceIDs() const;
+		UE_API EDeployDeviceOption GetDeployDeviceOption() const;
+		UE_API void SetDeployDeviceOption( EDeployDeviceOption DeployDeviceOption );
+		UE_API float GetDeployDeviceListHeight() const;
+		UE_API void SetDeployDeviceListHeight( float NewHeight );
+		UE_API void OnDeviceRemoved(FString DeviceID);
+		UE_API TSharedRef<SWidget> CreateDeployDeviceWidget();
 		float DeployDeviceListHeight = 150.0f;
 
-		void SetBuild( bool bBuild );
-		bool GetBuild( ILauncherProfilePtr InProfile = nullptr ) const;
+		UE_API void SetBuild( bool bBuild );
+		UE_API bool GetBuild( ILauncherProfilePtr InProfile = nullptr ) const;
 
-		void SetForceBuild( bool bForceBuild );
-		bool GetForceBuild( ILauncherProfilePtr InProfile = nullptr ) const;
+		UE_API void SetForceBuild( bool bForceBuild );
+		UE_API bool GetForceBuild( ILauncherProfilePtr InProfile = nullptr ) const;
 
-		void SetArchitecture( FString Architecture );
-		FString GetArchitecture() const;
-		FText GetArchitectureDisplayName( FString Architecture );
+		UE_API void SetArchitecture( FString Architecture );
+		UE_API FString GetArchitecture() const;
+		UE_API FText GetArchitectureDisplayName( FString Architecture );
 
-		void SetDeployToDevice( bool bDeployToDevice );
-		bool GetDeployToDevice( ILauncherProfilePtr InProfile = nullptr ) const;
+		UE_API void SetDeployToDevice( bool bDeployToDevice );
+		UE_API bool GetDeployToDevice( ILauncherProfilePtr InProfile = nullptr ) const;
 
-		void SetIsRunning( bool bRun );
-		bool GetIsRunning( ILauncherProfilePtr InProfile = nullptr ) const;
+		UE_API void SetIsRunning( bool bRun );
+		UE_API bool GetIsRunning( ILauncherProfilePtr InProfile = nullptr ) const;
 
-		void OnInitialMapChanged( TSharedPtr<FString> InitialMap, ESelectInfo::Type );
-		TSharedPtr<FString> GetInitialMap() const;
+		UE_API void OnInitialMapChanged( TSharedPtr<FString> InitialMap, ESelectInfo::Type );
+		UE_API TSharedPtr<FString> GetInitialMap() const;
 
 
 
 	protected:
 
-		virtual void OnPropertyChanged() override;
-		void RefreshContentScheme();
-		void CacheStartupMapList() const;
-		void CacheArchitectures();
-		TSharedRef<SWidget> OnGenerateComboWidget( TSharedPtr<FString> InComboString );
+		UE_API virtual void OnPropertyChanged() override;
+		UE_API void RefreshContentScheme();
+		UE_API void CacheStartupMapList() const;
+		UE_API void CacheArchitectures();
+		UE_API TSharedRef<SWidget> OnGenerateComboWidget( TSharedPtr<FString> InComboString );
 
 
 		FLaunchProfileTreeDataRef TreeData;
@@ -194,3 +196,5 @@ namespace ProjectLauncher
 		EBuildTargetType CachedBuildTargetType = EBuildTargetType::Game;
 	};
 }
+
+#undef UE_API
diff --git a/Engine/Plugins/Developer/ProjectLauncher/Source/ProjectLauncher/Public/ProfileTree/LaunchProfileTreeData.h b/Engine/Plugins/Developer/ProjectLauncher/Source/ProjectLauncher/Public/ProfileTree/LaunchProfileTreeData.h
--- a/Engine/Plugins/Developer/ProjectLauncher/Source/ProjectLauncher/Public/ProfileTree/LaunchProfileTreeData.h
+++ b/Engine/Plugins/Developer/ProjectLauncher/Source/ProjectLauncher/Public/ProfileTree/LaunchProfileTreeData.h
@@ -7,6 +7,8 @@
 #include "Templates/SharedPointer.h"
 #include "ILauncherProfile.h"
 
+#define UE_API PROJECTLAUNCHER_API
+
 
 class SWidget;
 struct EVisibility;
@@ -25,11 +27,11 @@ namespace ProjectLauncher
 	typedef TSharedRef<FLaunchProfileTreeNode> FLaunchProfileTreeNodeRef;
 
 
-	class PROJECTLAUNCHER_API FLaunchProfileTreeNode : public TSharedFromThis<FLaunchProfileTreeNode>
+	class FLaunchProfileTreeNode : public TSharedFromThis<FLaunchProfileTreeNode>
 	{
 	public:
 		// construction
-		FLaunchProfileTreeNode( FLaunchProfileTreeData* InTreeData );
+		UE_API FLaunchProfileTreeNode( FLaunchProfileTreeData* InTreeData );
 
 		typedef TFunction<void()> FFunction;
 		typedef TFunction<bool()> FGetBool;
@@ -45,8 +47,8 @@ namespace ProjectLauncher
 			FGetBool IsVisible = nullptr;
 			FGetBool IsEnabled = nullptr;
 		};
-		FLaunchProfileTreeNode& AddWidget( FText InName, FCallbacks&& InWidgetCallbacks, TSharedRef<SWidget> InValueWidget );
-		FLaunchProfileTreeNode& AddWidget( FText InName, TSharedRef<SWidget> InValueWidget );
+		UE_API FLaunchProfileTreeNode& AddWidget( FText InName, FCallbacks&& InWidgetCallbacks, TSharedRef<SWidget> InValueWidget );
+		UE_API FLaunchProfileTreeNode& AddWidget( FText InName, TSharedRef<SWidget> InValueWidget );
 
 
 		struct FBooleanCallbacks
@@ -57,7 +59,7 @@ namespace ProjectLauncher
 			FGetBool IsVisible = nullptr;
 			FGetBool IsEnabled = nullptr;
 		};
-		FLaunchProfileTreeNode& AddBoolean( FText InName, FBooleanCallbacks&& BooleanCallbacks );
+		UE_API FLaunchProfileTreeNode& AddBoolean( FText InName, FBooleanCallbacks&& BooleanCallbacks );
 
 		struct FStringCallbacks
 		{
@@ -68,9 +70,9 @@ namespace ProjectLauncher
 			FGetBool IsEnabled = nullptr;
 
 		};
-		FLaunchProfileTreeNode& AddString( FText InName, FStringCallbacks&& StringCallbacks );
-		FLaunchProfileTreeNode& AddDirectoryString( FText InName, FStringCallbacks&& StringCallbacks );
-		FLaunchProfileTreeNode& AddCommandLineString( FText InName, FStringCallbacks&& StringCallbacks );
+		UE_API FLaunchProfileTreeNode& AddString( FText InName, FStringCallbacks&& StringCallbacks );
+		UE_API FLaunchProfileTreeNode& AddDirectoryString( FText InName, FStringCallbacks&& StringCallbacks );
+		UE_API FLaunchProfileTreeNode& AddCommandLineString( FText InName, FStringCallbacks&& StringCallbacks );
 
 
 
@@ -90,17 +92,17 @@ namespace ProjectLauncher
 
 
 
-	class PROJECTLAUNCHER_API FLaunchProfileTreeData : public TSharedFromThis<FLaunchProfileTreeData>
+	class FLaunchProfileTreeData : public TSharedFromThis<FLaunchProfileTreeData>
 	{
 	public:
 		// construction
-		FLaunchProfileTreeData(ILauncherProfilePtr InProfile, TSharedRef<FModel> InModel, ILaunchProfileTreeBuilder* InTreeBuilder);
+		UE_API FLaunchProfileTreeData(ILauncherProfilePtr InProfile, TSharedRef<FModel> InModel, ILaunchProfileTreeBuilder* InTreeBuilder);
 
 		// node creation
-		FLaunchProfileTreeNode& AddHeading( FText InName );
-		void CreateExtensionsUI();
+		UE_API FLaunchProfileTreeNode& AddHeading( FText InName );
+		UE_API void CreateExtensionsUI();
 
-		void RequestTreeRefresh();
+		UE_API void RequestTreeRefresh();
 
 		// member variables
 		ILauncherProfilePtr Profile;
@@ -115,3 +117,4 @@ namespace ProjectLauncher
 
 }
 
+#undef UE_API
diff --git a/Engine/Plugins/Developer/ProjectLauncher/Source/ProjectLauncher/Public/ProjectLauncherModule.h b/Engine/Plugins/Developer/ProjectLauncher/Source/ProjectLauncher/Public/ProjectLauncherModule.h
--- a/Engine/Plugins/Developer/ProjectLauncher/Source/ProjectLauncher/Public/ProjectLauncherModule.h
+++ b/Engine/Plugins/Developer/ProjectLauncher/Source/ProjectLauncher/Public/ProjectLauncherModule.h
@@ -14,7 +14,7 @@ namespace ProjectLauncher
 /*
  * Interface for ProjectLauncher module
  */
-class PROJECTLAUNCHER_API IProjectLauncherModule : public IModuleInterface
+class IProjectLauncherModule : public IModuleInterface
 {
 public:
 	/**
diff --git a/Engine/Plugins/Developer/ProjectLauncher/Source/ProjectLauncher/Public/Widgets/Shared/SCustomLaunchContentSchemeCombo.h b/Engine/Plugins/Developer/ProjectLauncher/Source/ProjectLauncher/Public/Widgets/Shared/SCustomLaunchContentSchemeCombo.h
--- a/Engine/Plugins/Developer/ProjectLauncher/Source/ProjectLauncher/Public/Widgets/Shared/SCustomLaunchContentSchemeCombo.h
+++ b/Engine/Plugins/Developer/ProjectLauncher/Source/ProjectLauncher/Public/Widgets/Shared/SCustomLaunchContentSchemeCombo.h
@@ -10,7 +10,9 @@
 #include "Fonts/SlateFontInfo.h"
 #include "Model/ProjectLauncherModel.h"
 
-class PROJECTLAUNCHER_API SCustomLaunchContentSchemeCombo
+#define UE_API PROJECTLAUNCHER_API
+
+class SCustomLaunchContentSchemeCombo
 	: public SCompoundWidget
 {
 public:
@@ -28,15 +30,17 @@ public:
 	SLATE_END_ARGS()
 
 public:
-	void Construct(	const FArguments& InArgs);
+	UE_API void Construct(	const FArguments& InArgs);
 
 protected:
 	TAttribute<ProjectLauncher::EContentScheme> SelectedContentScheme;
 	FOnSelectionChanged OnSelectionChanged;
 	FGetContentSchemeAvailability IsContentSchemeAvailable;
 
-	TSharedRef<SWidget> MakeContentSchemeSelectionWidget();
+	UE_API TSharedRef<SWidget> MakeContentSchemeSelectionWidget();
 
-	FText GetContentSchemeName() const;
-	void SetContentScheme(ProjectLauncher::EContentScheme ContentScheme);
+	UE_API FText GetContentSchemeName() const;
+	UE_API void SetContentScheme(ProjectLauncher::EContentScheme ContentScheme);
 };
+
+#undef UE_API
diff --git a/Engine/Plugins/Developer/ProjectLauncher/Source/ProjectLauncher/Public/Widgets/Shared/SCustomLaunchDeviceCombo.h b/Engine/Plugins/Developer/ProjectLauncher/Source/ProjectLauncher/Public/Widgets/Shared/SCustomLaunchDeviceCombo.h
--- a/Engine/Plugins/Developer/ProjectLauncher/Source/ProjectLauncher/Public/Widgets/Shared/SCustomLaunchDeviceCombo.h
+++ b/Engine/Plugins/Developer/ProjectLauncher/Source/ProjectLauncher/Public/Widgets/Shared/SCustomLaunchDeviceCombo.h
@@ -5,9 +5,11 @@
 #include "CoreMinimal.h"
 #include "Widgets/Shared/SCustomLaunchDeviceWidgetBase.h"
 
+#define UE_API PROJECTLAUNCHER_API
+
 template<typename ItemType> class SComboBox;
 
-class PROJECTLAUNCHER_API SCustomLaunchDeviceCombo
+class SCustomLaunchDeviceCombo
 	: public SCustomLaunchDeviceWidgetBase
 {
 public:
@@ -22,16 +24,18 @@ public:
 	SLATE_END_ARGS()
 
 public:
-	void Construct(	const FArguments& InArgs );
+	UE_API void Construct(	const FArguments& InArgs );
 
 protected:
 
-	TSharedRef<SWidget> GenerateDeviceProxyListWidget(TSharedPtr<ITargetDeviceProxy> DeviceProxy) const;
-	void OnDeviceProxySelectionChangedChanged(TSharedPtr<ITargetDeviceProxy> DeviceProxy, ESelectInfo::Type InSelectInfo);
-	const FSlateBrush* GetSelectedDeviceProxyBrush() const;
-	FText GetSelectedDeviceProxyName() const;
+	UE_API TSharedRef<SWidget> GenerateDeviceProxyListWidget(TSharedPtr<ITargetDeviceProxy> DeviceProxy) const;
+	UE_API void OnDeviceProxySelectionChangedChanged(TSharedPtr<ITargetDeviceProxy> DeviceProxy, ESelectInfo::Type InSelectInfo);
+	UE_API const FSlateBrush* GetSelectedDeviceProxyBrush() const;
+	UE_API FText GetSelectedDeviceProxyName() const;
 
 private:
 	TSharedPtr<SComboBox<TSharedPtr<ITargetDeviceProxy>> > DeviceProxyComboBox;
 
 };
+
+#undef UE_API
diff --git a/Engine/Plugins/Developer/ProjectLauncher/Source/ProjectLauncher/Public/Widgets/Shared/SCustomLaunchDeviceListView.h b/Engine/Plugins/Developer/ProjectLauncher/Source/ProjectLauncher/Public/Widgets/Shared/SCustomLaunchDeviceListView.h
--- a/Engine/Plugins/Developer/ProjectLauncher/Source/ProjectLauncher/Public/Widgets/Shared/SCustomLaunchDeviceListView.h
+++ b/Engine/Plugins/Developer/ProjectLauncher/Source/ProjectLauncher/Public/Widgets/Shared/SCustomLaunchDeviceListView.h
@@ -5,11 +5,13 @@
 #include "CoreMinimal.h"
 #include "Widgets/Shared/SCustomLaunchDeviceWidgetBase.h"
 
+#define UE_API PROJECTLAUNCHER_API
+
 class ITableRow;
 class STableViewBase;
 template <typename ItemType> class SListView;
 
-class PROJECTLAUNCHER_API SCustomLaunchDeviceListView
+class SCustomLaunchDeviceListView
 	: public SCustomLaunchDeviceWidgetBase
 {
 public:
@@ -26,17 +28,19 @@ public:
 	SLATE_END_ARGS()
 
 public:
-	void Construct(	const FArguments& InArgs );
+	UE_API void Construct(	const FArguments& InArgs );
 
 protected:
 	bool bSingleSelect;
 
-	TSharedRef<ITableRow> GenerateDeviceProxyRow(TSharedPtr<ITargetDeviceProxy> DeviceProxy, const TSharedRef<STableViewBase>& OwnerTable);
-	ECheckBoxState IsDeviceProxyChecked(TSharedPtr<ITargetDeviceProxy> DeviceProxy) const;
-	void OnDeviceProxyCheckStateChanged(ECheckBoxState NewState, TSharedPtr<ITargetDeviceProxy> DeviceProxy);
-	virtual void OnDeviceListRefreshed() override;
+	UE_API TSharedRef<ITableRow> GenerateDeviceProxyRow(TSharedPtr<ITargetDeviceProxy> DeviceProxy, const TSharedRef<STableViewBase>& OwnerTable);
+	UE_API ECheckBoxState IsDeviceProxyChecked(TSharedPtr<ITargetDeviceProxy> DeviceProxy) const;
+	UE_API void OnDeviceProxyCheckStateChanged(ECheckBoxState NewState, TSharedPtr<ITargetDeviceProxy> DeviceProxy);
+	UE_API virtual void OnDeviceListRefreshed() override;
 
 private:
 	TSharedPtr<SListView<TSharedPtr<ITargetDeviceProxy>> > DeviceProxyListView;
 
 };
+
+#undef UE_API
diff --git a/Engine/Plugins/Developer/ProjectLauncher/Source/ProjectLauncher/Public/Widgets/Shared/SCustomLaunchDeviceWidgetBase.h b/Engine/Plugins/Developer/ProjectLauncher/Source/ProjectLauncher/Public/Widgets/Shared/SCustomLaunchDeviceWidgetBase.h
--- a/Engine/Plugins/Developer/ProjectLauncher/Source/ProjectLauncher/Public/Widgets/Shared/SCustomLaunchDeviceWidgetBase.h
+++ b/Engine/Plugins/Developer/ProjectLauncher/Source/ProjectLauncher/Public/Widgets/Shared/SCustomLaunchDeviceWidgetBase.h
@@ -5,10 +5,12 @@
 #include "CoreMinimal.h"
 #include "Widgets/SCompoundWidget.h"
 
+#define UE_API PROJECTLAUNCHER_API
+
 class ITargetDeviceProxy;
 class ITargetDeviceProxyManager;
 
-class PROJECTLAUNCHER_API SCustomLaunchDeviceWidgetBase
+class SCustomLaunchDeviceWidgetBase
 	: public SCompoundWidget
 {
 public:
@@ -16,11 +18,11 @@ public:
 	DECLARE_DELEGATE_OneParam(FOnDeviceRemoved, FString );
 
 public:
-	void Construct();
-	~SCustomLaunchDeviceWidgetBase();
+	UE_API void Construct();
+	UE_API ~SCustomLaunchDeviceWidgetBase();
 
-	void RefreshDeviceList();
-	void OnSelectedPlatformChanged();
+	UE_API void RefreshDeviceList();
+	UE_API void OnSelectedPlatformChanged();
 
 protected:
 	virtual void OnDeviceListRefreshed() {};
@@ -31,10 +33,12 @@ protected:
 	FOnDeviceRemoved OnDeviceRemoved;
 	bool bAllPlatforms = false;
 
-	void OnDeviceProxyAdded(const TSharedRef<ITargetDeviceProxy>& DeviceProxy);
-	void OnDeviceProxyRemoved(const TSharedRef<ITargetDeviceProxy>& DeviceProxy);
+	UE_API void OnDeviceProxyAdded(const TSharedRef<ITargetDeviceProxy>& DeviceProxy);
+	UE_API void OnDeviceProxyRemoved(const TSharedRef<ITargetDeviceProxy>& DeviceProxy);
 
-	const TSharedRef<ITargetDeviceProxyManager> GetDeviceProxyManager() const;
+	UE_API const TSharedRef<ITargetDeviceProxyManager> GetDeviceProxyManager() const;
 
 	TArray<TSharedPtr<ITargetDeviceProxy>> DeviceProxyList;
 };
+
+#undef UE_API
diff --git a/Engine/Plugins/Developer/ProjectLauncher/Source/ProjectLauncher/Public/Widgets/Shared/SCustomLaunchMapListView.h b/Engine/Plugins/Developer/ProjectLauncher/Source/ProjectLauncher/Public/Widgets/Shared/SCustomLaunchMapListView.h
--- a/Engine/Plugins/Developer/ProjectLauncher/Source/ProjectLauncher/Public/Widgets/Shared/SCustomLaunchMapListView.h
+++ b/Engine/Plugins/Developer/ProjectLauncher/Source/ProjectLauncher/Public/Widgets/Shared/SCustomLaunchMapListView.h
@@ -6,11 +6,13 @@
 #include "Widgets/SCompoundWidget.h"
 #include "Model/ProjectLauncherModel.h"
 
+#define UE_API PROJECTLAUNCHER_API
+
 class ITableRow;
 class STableViewBase;
 template<typename ItemType> class STreeView;
 
-class PROJECTLAUNCHER_API SCustomLaunchMapListView
+class SCustomLaunchMapListView
 	: public SCompoundWidget
 {
 public:
@@ -23,10 +25,10 @@ public:
 	SLATE_END_ARGS()
 
 public:
-	void Construct( const FArguments& InArgs, TSharedRef<ProjectLauncher::FModel> InModel );
+	UE_API void Construct( const FArguments& InArgs, TSharedRef<ProjectLauncher::FModel> InModel );
 
-	void RefreshMapList();
-	TSharedRef<SWidget> MakeControlsWidget();
+	UE_API void RefreshMapList();
+	UE_API TSharedRef<SWidget> MakeControlsWidget();
 
 protected:
 	TSharedPtr<ProjectLauncher::FModel> Model;
@@ -48,32 +50,34 @@ private:
 
 	FMapTreeNodePtr MapTreeRoot;
 
-	void OnProjectChanged();
+	UE_API void OnProjectChanged();
 
 
-	void RefreshCheckBoxState( bool bExpand = false);
-	ECheckBoxState RefreshCheckBoxStateRecursive(FMapTreeNodePtr Node, bool bExpand);
-	void SetCheckBoxStateRecursive(FMapTreeNodePtr Node, ECheckBoxState CheckBoxState, TArray<FString>& CheckedMaps);
+	UE_API void RefreshCheckBoxState( bool bExpand = false);
+	UE_API ECheckBoxState RefreshCheckBoxStateRecursive(FMapTreeNodePtr Node, bool bExpand);
+	UE_API void SetCheckBoxStateRecursive(FMapTreeNodePtr Node, ECheckBoxState CheckBoxState, TArray<FString>& CheckedMaps);
 
-	void OnSearchFilterTextCommitted(const FText& SearchText, ETextCommit::Type InCommitType);
-	void OnSearchFilterTextChanged(const FText& SearchText);
+	UE_API void OnSearchFilterTextCommitted(const FText& SearchText, ETextCommit::Type InCommitType);
+	UE_API void OnSearchFilterTextChanged(const FText& SearchText);
 
 	FString CurrentFilterText;
 
 	TSharedPtr<STreeView<FMapTreeNodePtr>> MapTreeView;
 	TArray<FMapTreeNodePtr> MapTreeViewItemsSource;
 
-	void GetMapTreeNodeChildren(FMapTreeNodePtr Item, TArray<FMapTreeNodePtr>& OutChildren);
-	TSharedRef<ITableRow> GenerateMapTreeNodeRow(FMapTreeNodePtr Item, const TSharedRef<STableViewBase>& OwnerTable);
-	ECheckBoxState GetMapTreeNodeCheckState(FMapTreeNodePtr Item) const;
-	void SetMapTreeNodeCheckState(ECheckBoxState CheckBoxState, FMapTreeNodePtr Item);
-	const FSlateBrush* GetMapTreeNodeIcon(FMapTreeNodePtr Node) const;
-	FSlateColor GetMapTreeNodeColor(FMapTreeNodePtr Node) const;
+	UE_API void GetMapTreeNodeChildren(FMapTreeNodePtr Item, TArray<FMapTreeNodePtr>& OutChildren);
+	UE_API TSharedRef<ITableRow> GenerateMapTreeNodeRow(FMapTreeNodePtr Item, const TSharedRef<STableViewBase>& OwnerTable);
+	UE_API ECheckBoxState GetMapTreeNodeCheckState(FMapTreeNodePtr Item) const;
+	UE_API void SetMapTreeNodeCheckState(ECheckBoxState CheckBoxState, FMapTreeNodePtr Item);
+	UE_API const FSlateBrush* GetMapTreeNodeIcon(FMapTreeNodePtr Node) const;
+	UE_API FSlateColor GetMapTreeNodeColor(FMapTreeNodePtr Node) const;
 
 	mutable bool bHasPaintedThisFrame = false;
-	virtual int32 OnPaint( const FPaintArgs& Args, const FGeometry& AllottedGeometry, const FSlateRect& MyCullingRect, FSlateWindowElementList& OutDrawElements, int32 LayerId, const FWidgetStyle& InWidgetStyle, bool bParentEnabled) const override;
+	UE_API virtual int32 OnPaint( const FPaintArgs& Args, const FGeometry& AllottedGeometry, const FSlateRect& MyCullingRect, FSlateWindowElementList& OutDrawElements, int32 LayerId, const FWidgetStyle& InWidgetStyle, bool bParentEnabled) const override;
 
 	bool bMapListDirty = false;
-	void RefreshMapListInternal();
-	virtual void Tick(const FGeometry& AllottedGeometry, const double InCurrentTime, const float InDeltaTime) override;
+	UE_API void RefreshMapListInternal();
+	UE_API virtual void Tick(const FGeometry& AllottedGeometry, const double InCurrentTime, const float InDeltaTime) override;
 };
+
+#undef UE_API
diff --git a/Engine/Plugins/Developer/ProjectLauncher/Source/ProjectLauncher/Public/Widgets/Shared/SCustomLaunchPlatformCombo.h b/Engine/Plugins/Developer/ProjectLauncher/Source/ProjectLauncher/Public/Widgets/Shared/SCustomLaunchPlatformCombo.h
--- a/Engine/Plugins/Developer/ProjectLauncher/Source/ProjectLauncher/Public/Widgets/Shared/SCustomLaunchPlatformCombo.h
+++ b/Engine/Plugins/Developer/ProjectLauncher/Source/ProjectLauncher/Public/Widgets/Shared/SCustomLaunchPlatformCombo.h
@@ -8,10 +8,12 @@
 #include "Styling/SlateWidgetStyleAsset.h"
 #include "Fonts/SlateFontInfo.h"
 
+#define UE_API PROJECTLAUNCHER_API
+
 template<typename ItemType> class SComboBox;
 
 
-class PROJECTLAUNCHER_API SCustomLaunchPlatformCombo
+class SCustomLaunchPlatformCombo
 	: public SCompoundWidget
 {
 public:
@@ -28,17 +30,19 @@ public:
 	SLATE_END_ARGS()
 
 public:
-	void Construct(	const FArguments& InArgs);
+	UE_API void Construct(	const FArguments& InArgs);
 
 protected:
 	TAttribute<TArray<FString>> SelectedPlatforms;
 	FOnSelectionChanged OnSelectionChanged;
 	bool bBasicPlatformsOnly;
 
-	TSharedRef<SWidget> OnGeneratePlatformListWidget( TSharedPtr<FString> Platform ) const;
-	void OnPlatformSelectionChanged( TSharedPtr<FString> Platform, ESelectInfo::Type InSelectInfo );
-	const FSlateBrush* GetSelectedPlatformBrush() const;
-	FText GetSelectedPlatformName() const;
+	UE_API TSharedRef<SWidget> OnGeneratePlatformListWidget( TSharedPtr<FString> Platform ) const;
+	UE_API void OnPlatformSelectionChanged( TSharedPtr<FString> Platform, ESelectInfo::Type InSelectInfo );
+	UE_API const FSlateBrush* GetSelectedPlatformBrush() const;
+	UE_API FText GetSelectedPlatformName() const;
 	TArray<TSharedPtr<FString>> PlatformsList;
 	TSharedPtr<SComboBox<TSharedPtr<FString>>> PlatformsComboBox;
 };
+
+#undef UE_API
diff --git a/Engine/Plugins/Developer/ProjectLauncher/Source/ProjectLauncher/Public/Widgets/Shared/SCustomLaunchProjectCombo.h b/Engine/Plugins/Developer/ProjectLauncher/Source/ProjectLauncher/Public/Widgets/Shared/SCustomLaunchProjectCombo.h
--- a/Engine/Plugins/Developer/ProjectLauncher/Source/ProjectLauncher/Public/Widgets/Shared/SCustomLaunchProjectCombo.h
+++ b/Engine/Plugins/Developer/ProjectLauncher/Source/ProjectLauncher/Public/Widgets/Shared/SCustomLaunchProjectCombo.h
@@ -9,7 +9,9 @@
 #include "Styling/SlateWidgetStyleAsset.h"
 #include "Fonts/SlateFontInfo.h"
 
-class PROJECTLAUNCHER_API SCustomLaunchProjectCombo
+#define UE_API PROJECTLAUNCHER_API
+
+class SCustomLaunchProjectCombo
 	: public SCompoundWidget
 {
 public:
@@ -38,7 +40,7 @@ public:
 	SLATE_END_ARGS()
 
 public:
-	void Construct(	const FArguments& InArgs);
+	UE_API void Construct(	const FArguments& InArgs);
 
 protected:
 	TAttribute<FString> SelectedProject;
@@ -47,11 +49,13 @@ protected:
 	bool bShowAnyProjectOption;
 	ECurrentProjectOption CurrentProjectOption;
 
-	TSharedRef<SWidget> MakeProjectSelectionWidget();
-	void OnBrowseForProject();
+	UE_API TSharedRef<SWidget> MakeProjectSelectionWidget();
+	UE_API void OnBrowseForProject();
 
-	FText GetProjectName() const;
-	void SetProjectPath(FString ProjectPath);
+	UE_API FText GetProjectName() const;
+	UE_API void SetProjectPath(FString ProjectPath);
 
 
 };
+
+#undef UE_API
diff --git a/Engine/Plugins/Editor/AssetManagerEditor/Source/AssetManagerEditor/Public/AssetManagerEditorModule.h b/Engine/Plugins/Editor/AssetManagerEditor/Source/AssetManagerEditor/Public/AssetManagerEditorModule.h
--- a/Engine/Plugins/Editor/AssetManagerEditor/Source/AssetManagerEditor/Public/AssetManagerEditorModule.h
+++ b/Engine/Plugins/Editor/AssetManagerEditor/Source/AssetManagerEditor/Public/AssetManagerEditorModule.h
@@ -11,6 +11,8 @@
 #include "Engine/AssetManagerTypes.h"
 #include "CollectionManagerTypes.h"
 
+#define UE_API ASSETMANAGEREDITOR_API
+
 class ICollectionContainer;
 class FUICommandList;
 class FUICommandInfo;
@@ -52,7 +54,7 @@ class SWidget;
 struct FAssetData;
 
 /** Struct containing the information about currently investigated asset context */
-struct ASSETMANAGEREDITOR_API FAssetManagerEditorRegistrySource
+struct FAssetManagerEditorRegistrySource
 {
 	/** Display name of the source, equal to target Platform, Editor, or Custom */
 	FString SourceName;
@@ -86,10 +88,10 @@ public:
 	TMap<int32, FAssetManagerChunkInfo> ChunkAssignments;
 
 	/** If SourceName matches this, this is the editor */
-	static const FString EditorSourceName;
+	static UE_API const FString EditorSourceName;
 
 	/** If SourceName matches this, this is the custom source loaded from an arbitrary path */
-	static const FString CustomSourceName;
+	static UE_API const FString CustomSourceName;
 
 	FAssetManagerEditorRegistrySource()
 		: TargetPlatform(nullptr)
@@ -198,7 +200,7 @@ public:
 	}
 
 	// grabs the SourceTimestamp from the asset registry file changed times.
-	void LoadRegistryTimestamp();
+	UE_API void LoadRegistryTimestamp();
 
 	// Functions to forward on to the RegistryState, either the owned one or the editor's global State.
 	FAssetData GetAssetByObjectPath(const FSoftObjectPath& ObjectPath) const
@@ -295,7 +297,7 @@ public:
 /**
  * The Asset Manager Editor module handles creating UI for asset management and exposes several commands
  */
-class ASSETMANAGEREDITOR_API IAssetManagerEditorModule : public IModuleInterface
+class IAssetManagerEditorModule : public IModuleInterface
 {
 public:
 	/**
@@ -326,7 +328,7 @@ public:
 	 * @param bAllowClear If true, add None option to top
 	 * @param bAlowAll If true, add All Types option to bottom, returns AllPrimaryAssetTypes if selected
 	 */
-	static TSharedRef<SWidget> MakePrimaryAssetTypeSelector(FOnGetPrimaryAssetDisplayText OnGetDisplayText, FOnSetPrimaryAssetType OnSetType, bool bAllowClear = true, bool bAllowAll = false);
+	static UE_API TSharedRef<SWidget> MakePrimaryAssetTypeSelector(FOnGetPrimaryAssetDisplayText OnGetDisplayText, FOnSetPrimaryAssetType OnSetType, bool bAllowClear = true, bool bAllowAll = false);
 
 	/** 
 	 * Creates a simple version of a Primary Asset Id selector, not bound to a PropertyHandle
@@ -334,49 +336,49 @@ public:
 	 * @param OnSetId Delegate called when id is changed
 	 * @param bAllowClear If true, add None option to top
 	 */
-	static TSharedRef<SWidget> MakePrimaryAssetIdSelector(FOnGetPrimaryAssetDisplayText OnGetDisplayText, FOnSetPrimaryAssetId OnSetId, bool bAllowClear = true, 
+	static UE_API TSharedRef<SWidget> MakePrimaryAssetIdSelector(FOnGetPrimaryAssetDisplayText OnGetDisplayText, FOnSetPrimaryAssetId OnSetId, bool bAllowClear = true, 
 		TArray<FPrimaryAssetType> AllowedTypes = TArray<FPrimaryAssetType>(), TArray<const UClass*> AllowedClasses = TArray<const UClass*>(), TArray<const UClass*> DisallowedClasses = TArray<const UClass*>()
 	);
 
 	/** Called to get list of valid primary asset types */
-	static void GeneratePrimaryAssetTypeComboBoxStrings(TArray< TSharedPtr<FString> >& OutComboBoxStrings, TArray<TSharedPtr<SToolTip>>& OutToolTips, TArray<bool>& OutRestrictedItems, bool bAllowClear, bool bAllowAll);
+	static UE_API void GeneratePrimaryAssetTypeComboBoxStrings(TArray< TSharedPtr<FString> >& OutComboBoxStrings, TArray<TSharedPtr<SToolTip>>& OutToolTips, TArray<bool>& OutRestrictedItems, bool bAllowClear, bool bAllowAll);
 
 	/** Called for asset picker to see rather to show asset */
-	static bool OnShouldFilterPrimaryAsset(const FAssetData& InAssetData, TArray<FPrimaryAssetType> AllowedTypes);
+	static UE_API bool OnShouldFilterPrimaryAsset(const FAssetData& InAssetData, TArray<FPrimaryAssetType> AllowedTypes);
 
 	/** Creates fake AssetData to represent chunks */
-	static FAssetData CreateFakeAssetDataFromChunkId(int32 ChunkID);
+	static UE_API FAssetData CreateFakeAssetDataFromChunkId(int32 ChunkID);
 
 	/** Returns >= 0 if this is a fake chunk AssetData, otherwise INDEX_NONE */
-	static int32 ExtractChunkIdFromFakeAssetData(const FAssetData& InAssetData);
+	static UE_API int32 ExtractChunkIdFromFakeAssetData(const FAssetData& InAssetData);
 
 	/** Creates fake AssetData to represent a primary asset */
-	static FAssetData CreateFakeAssetDataFromPrimaryAssetId(const FPrimaryAssetId& PrimaryAssetId);
+	static UE_API FAssetData CreateFakeAssetDataFromPrimaryAssetId(const FPrimaryAssetId& PrimaryAssetId);
 
 	/** Returns valid PrimaryAssetId if this is a fake AssetData for a primary asset */
-	static FPrimaryAssetId ExtractPrimaryAssetIdFromFakeAssetData(const FAssetData& InAssetData);
+	static UE_API FPrimaryAssetId ExtractPrimaryAssetIdFromFakeAssetData(const FAssetData& InAssetData);
 
 	/** Parses a list of AssetData and extracts AssetIdentifiers, handles the fake asset data */
-	static void ExtractAssetIdentifiersFromAssetDataList(const TArray<FAssetData>& AssetDataList, TArray<FAssetIdentifier>& OutAssetIdentifiers);
+	static UE_API void ExtractAssetIdentifiersFromAssetDataList(const TArray<FAssetData>& AssetDataList, TArray<FAssetIdentifier>& OutAssetIdentifiers);
 
 	/** Fake package path used by placeholder asset data, has primary asset ID added to it to set package name */
-	static const FName PrimaryAssetFakeAssetDataPackagePath;
+	static UE_API const FName PrimaryAssetFakeAssetDataPackagePath;
 
 	/** Specific package path inside PrimaryAssetFakeAssetDataPackagePath representing chunks */
-	static const FName ChunkFakeAssetDataPackageName;
+	static UE_API const FName ChunkFakeAssetDataPackageName;
 	
 	/** Fake type returned from type combo box to mean all types */
-	static const FPrimaryAssetType AllPrimaryAssetTypes;
+	static UE_API const FPrimaryAssetType AllPrimaryAssetTypes;
 
 	/** Custom column names */
-	static const FName ResourceSizeName;
-	static const FName DiskSizeName;
-	static const FName ManagedResourceSizeName;
-	static const FName ManagedDiskSizeName;
-	static const FName TotalUsageName;
-	static const FName CookRuleName;
-	static const FName ChunksName;
-	static const FName PluginName;
+	static UE_API const FName ResourceSizeName;
+	static UE_API const FName DiskSizeName;
+	static UE_API const FName ManagedResourceSizeName;
+	static UE_API const FName ManagedDiskSizeName;
+	static UE_API const FName TotalUsageName;
+	static UE_API const FName CookRuleName;
+	static UE_API const FName ChunksName;
+	static UE_API const FName PluginName;
 
 	/** Gets the value of a "virtual" column for an asset data, this will query the AssetManager for you and takes current platform into account. Returns true and sets out parameter if found */
 	virtual bool GetStringValueForCustomColumn(const FAssetData& AssetData, FName ColumnName, FString& OutValue, const FAssetManagerEditorRegistrySource* OverrideRegistrySource = nullptr) = 0;
@@ -454,3 +456,5 @@ public:
 	/* Return a UToolsMenu generated widget containing the asset source control context menu. */
 	virtual TSharedRef<SWidget> CreateAssetSourceControlContextMenu(const TArray<FAssetData>& SelectedAssets) = 0;
 };
+
+#undef UE_API
diff --git a/Engine/Plugins/Editor/AssetManagerEditor/Source/AssetManagerEditor/Public/ReferenceViewer/EdGraphNode_Reference.h b/Engine/Plugins/Editor/AssetManagerEditor/Source/AssetManagerEditor/Public/ReferenceViewer/EdGraphNode_Reference.h
--- a/Engine/Plugins/Editor/AssetManagerEditor/Source/AssetManagerEditor/Public/ReferenceViewer/EdGraphNode_Reference.h
+++ b/Engine/Plugins/Editor/AssetManagerEditor/Source/AssetManagerEditor/Public/ReferenceViewer/EdGraphNode_Reference.h
@@ -9,62 +9,64 @@
 #include "EdGraph_ReferenceViewer.h"
 #include "EdGraphNode_Reference.generated.h"
 
+#define UE_API ASSETMANAGEREDITOR_API
+
 class UEdGraphPin;
 
-UCLASS()
-class ASSETMANAGEREDITOR_API UEdGraphNode_Reference : public UEdGraphNode
+UCLASS(MinimalAPI)
+class UEdGraphNode_Reference : public UEdGraphNode
 {
 	GENERATED_UCLASS_BODY()
 
 	/** Returns first asset identifier */
-	FAssetIdentifier GetIdentifier() const;
+	UE_API FAssetIdentifier GetIdentifier() const;
 	
 	/** Returns all identifiers on this node including virtual things */
-	void GetAllIdentifiers(TArray<FAssetIdentifier>& OutIdentifiers) const;
+	UE_API void GetAllIdentifiers(TArray<FAssetIdentifier>& OutIdentifiers) const;
 
 	/** Returns only the packages in this node, skips searchable names */
-	void GetAllPackageNames(TArray<FName>& OutPackageNames) const;
+	UE_API void GetAllPackageNames(TArray<FName>& OutPackageNames) const;
 
 	/** Returns our owning graph */
-	UEdGraph_ReferenceViewer* GetReferenceViewerGraph() const;
+	UE_API UEdGraph_ReferenceViewer* GetReferenceViewerGraph() const;
 
 	// UEdGraphNode implementation
-	virtual FText GetNodeTitle(ENodeTitleType::Type TitleType) const override;
-	FLinearColor GetNodeTitleColor() const override;
-	virtual FText GetTooltipText() const override;
-	virtual void AllocateDefaultPins() override;
-	virtual UObject* GetJumpTargetForDoubleClick() const override;
-	virtual FSlateIcon GetIconAndTint(FLinearColor& OutColor) const override;
+	UE_API virtual FText GetNodeTitle(ENodeTitleType::Type TitleType) const override;
+	UE_API FLinearColor GetNodeTitleColor() const override;
+	UE_API virtual FText GetTooltipText() const override;
+	UE_API virtual void AllocateDefaultPins() override;
+	UE_API virtual UObject* GetJumpTargetForDoubleClick() const override;
+	UE_API virtual FSlateIcon GetIconAndTint(FLinearColor& OutColor) const override;
 	virtual bool ShowPaletteIconOnNode() const override { return true; }
 	// End UEdGraphNode implementation
 
 	void SetAllowThumbnail(bool bInAllow) { bAllowThumbnail = bInAllow; }
-	bool AllowsThumbnail() const;
-	bool UsesThumbnail() const;
-	bool IsPackage() const;
-	bool IsCollapsed() const;
+	UE_API bool AllowsThumbnail() const;
+	UE_API bool UsesThumbnail() const;
+	UE_API bool IsPackage() const;
+	UE_API bool IsCollapsed() const;
 	bool IsADuplicate() const { return bIsADuplicate; }
 
 	// Nodes that are filtered out still may still show because they
 	// are between nodes that pass the filter and the root.  This "filtered"
 	// bool allows us to render these in-between nodes differently
-	void SetIsFiltered(bool bInFiltered);
-	bool GetIsFiltered() const;
+	UE_API void SetIsFiltered(bool bInFiltered);
+	UE_API bool GetIsFiltered() const;
 
 	bool IsOverflow() const { return bIsOverflow; }
 
-	const FAssetData& GetAssetData() const;
+	UE_API const FAssetData& GetAssetData() const;
 
-	UEdGraphPin* GetDependencyPin();
-	UEdGraphPin* GetReferencerPin();
+	UE_API UEdGraphPin* GetDependencyPin();
+	UE_API UEdGraphPin* GetReferencerPin();
 
-	UE_INTERNAL void UpdatePath();
+	UE_API UE_INTERNAL void UpdatePath();
 
 private:
-	void CacheAssetData(const FAssetData& AssetData);
-	void SetupReferenceNode(const FIntPoint& NodeLoc, const TArray<FAssetIdentifier>& NewIdentifiers, const FAssetData& InAssetData, bool bInAllowThumbnail, bool bInIsADuplicate);
-	void SetReferenceNodeCollapsed(const FIntPoint& NodeLoc, int32 InNumReferencesExceedingMax, const TArray<FAssetIdentifier>& Identifiers);
-	void AddReferencer(class UEdGraphNode_Reference* ReferencerNode);
+	UE_API void CacheAssetData(const FAssetData& AssetData);
+	UE_API void SetupReferenceNode(const FIntPoint& NodeLoc, const TArray<FAssetIdentifier>& NewIdentifiers, const FAssetData& InAssetData, bool bInAllowThumbnail, bool bInIsADuplicate);
+	UE_API void SetReferenceNodeCollapsed(const FIntPoint& NodeLoc, int32 InNumReferencesExceedingMax, const TArray<FAssetIdentifier>& Identifiers);
+	UE_API void AddReferencer(class UEdGraphNode_Reference* ReferencerNode);
 
 	TArray<FAssetIdentifier> Identifiers;
 	FText NodeTitle;
@@ -89,3 +91,4 @@ private:
 };
 
 
+#undef UE_API
diff --git a/Engine/Plugins/Editor/AssetManagerEditor/Source/AssetManagerEditor/Public/ReferenceViewer/EdGraphNode_ReferencedProperties.h b/Engine/Plugins/Editor/AssetManagerEditor/Source/AssetManagerEditor/Public/ReferenceViewer/EdGraphNode_ReferencedProperties.h
--- a/Engine/Plugins/Editor/AssetManagerEditor/Source/AssetManagerEditor/Public/ReferenceViewer/EdGraphNode_ReferencedProperties.h
+++ b/Engine/Plugins/Editor/AssetManagerEditor/Source/AssetManagerEditor/Public/ReferenceViewer/EdGraphNode_ReferencedProperties.h
@@ -7,6 +7,8 @@
 #include "UObject/ObjectMacros.h"
 #include "EdGraphNode_ReferencedProperties.generated.h"
 
+#define UE_API ASSETMANAGEREDITOR_API
+
 /**
  * Describes a graph node property referencing an asset from another node
  */
@@ -88,8 +90,8 @@ using FReferencingPropertyDescriptionPtr = TSharedPtr<FReferencingPropertyDescri
 /**
  * A node to display a list of Node properties which are referencing another Node/Asset
  */
-UCLASS()
-class ASSETMANAGEREDITOR_API UEdGraphNode_ReferencedProperties : public UEdGraphNode
+UCLASS(MinimalAPI)
+class UEdGraphNode_ReferencedProperties : public UEdGraphNode
 {
 	GENERATED_BODY()
 
@@ -103,19 +105,19 @@ public:
 
 	const TObjectPtr<UEdGraphNode_Reference>& GetReferencedNode() const { return ReferencedNode; }
 
-	UObject* GetReferencingObject() const;
+	UE_API UObject* GetReferencingObject() const;
 	
-	UObject* GetReferencedObject() const;
+	UE_API UObject* GetReferencedObject() const;
 
 	/** Initialize this Node */
-	void SetupReferencedPropertiesNode(const TArray<FReferencingPropertyDescription>& InPropertiesDescription
+	UE_API void SetupReferencedPropertiesNode(const TArray<FReferencingPropertyDescription>& InPropertiesDescription
 		, const TObjectPtr<UEdGraphNode_Reference>& InReferencingNode, const TObjectPtr<UEdGraphNode_Reference>& InReferencedNode);
 
 	/**
 	 * Refresh the node location, so it stays mid-way between Referencing and Referenced nodes
 	 * @param InNodeSize: size of the widget representing this node. Useful when centering
 	 */
-	void RefreshLocation(const FVector2f& InNodeSize = FVector2f::Zero());
+	UE_API void RefreshLocation(const FVector2f& InNodeSize = FVector2f::Zero());
 
 	DECLARE_MULTICAST_DELEGATE(FOnPropertiesDescriptionUpdated)
 	FOnPropertiesDescriptionUpdated& OnPropertiesDescriptionUpdated() { return OnPropertiesDescriptionUpdatedDelegate; }
@@ -131,3 +133,5 @@ private:
 
 	FOnPropertiesDescriptionUpdated OnPropertiesDescriptionUpdatedDelegate;
 };
+
+#undef UE_API
diff --git a/Engine/Plugins/Editor/AssetManagerEditor/Source/AssetManagerEditor/Public/ReferenceViewer/EdGraph_ReferenceViewer.h b/Engine/Plugins/Editor/AssetManagerEditor/Source/AssetManagerEditor/Public/ReferenceViewer/EdGraph_ReferenceViewer.h
--- a/Engine/Plugins/Editor/AssetManagerEditor/Source/AssetManagerEditor/Public/ReferenceViewer/EdGraph_ReferenceViewer.h
+++ b/Engine/Plugins/Editor/AssetManagerEditor/Source/AssetManagerEditor/Public/ReferenceViewer/EdGraph_ReferenceViewer.h
@@ -11,6 +11,8 @@
 #include "Misc/FilterCollection.h"
 #include "EdGraph_ReferenceViewer.generated.h"
 
+#define UE_API ASSETMANAGEREDITOR_API
+
 class FAssetThumbnailPool;
 class UEdGraphNode_Reference;
 class UEdGraphNode_ReferencedProperties;
@@ -67,63 +69,63 @@ struct FReferenceNodeInfo
 
 
 
-UCLASS()
-class ASSETMANAGEREDITOR_API UEdGraph_ReferenceViewer : public UEdGraph
+UCLASS(MinimalAPI)
+class UEdGraph_ReferenceViewer : public UEdGraph
 {
 	GENERATED_UCLASS_BODY()
 
 public:
 	// UObject implementation
-	virtual void BeginDestroy() override;
+	UE_API virtual void BeginDestroy() override;
 	// End UObject implementation
 
 	/** Set reference viewer to focus on these assets */
-	void SetGraphRoot(const TArray<FAssetIdentifier>& GraphRootIdentifiers, const FIntPoint& GraphRootOrigin = FIntPoint(ForceInitToZero));
+	UE_API void SetGraphRoot(const TArray<FAssetIdentifier>& GraphRootIdentifiers, const FIntPoint& GraphRootOrigin = FIntPoint(ForceInitToZero));
 
 	/** Returns list of currently focused assets */
-	const TArray<FAssetIdentifier>& GetCurrentGraphRootIdentifiers() const;
+	UE_API const TArray<FAssetIdentifier>& GetCurrentGraphRootIdentifiers() const;
 
 	/** If you're extending the reference viewer via GetAllGraphEditorContextMenuExtender you can use this to get the list of selected assets to use in your menu extender */
-	bool GetSelectedAssetsForMenuExtender(const class UEdGraphNode* Node, TArray<FAssetIdentifier>& SelectedAssets) const;
+	UE_API bool GetSelectedAssetsForMenuExtender(const class UEdGraphNode* Node, TArray<FAssetIdentifier>& SelectedAssets) const;
 
 	/** Accessor for the thumbnail pool in this graph */
-	const TSharedPtr<class FAssetThumbnailPool>& GetAssetThumbnailPool() const;
+	UE_API const TSharedPtr<class FAssetThumbnailPool>& GetAssetThumbnailPool() const;
 
 	/** Force the graph to rebuild */
-	class UEdGraphNode_Reference* RebuildGraph();
+	UE_API class UEdGraphNode_Reference* RebuildGraph();
 
 	/** Refilters the nodes, more efficient that a full rebuild.  This function is preferred when the assets, reference types or depth hasn't changed, meaning the NodeInfos didn't change, just 
 	 * the presentation or filtering */
-	class UEdGraphNode_Reference* RefilterGraph();
+	UE_API class UEdGraphNode_Reference* RefilterGraph();
 
 	UE_INTERNAL bool IsShowingContentVersePath() const { return bShowingContentVersePath; }
-	UE_INTERNAL void SetShowingContentVersePath(bool bInShowingContentVersePath);
-	UE_INTERNAL void UpdatePaths();
+	UE_API UE_INTERNAL void SetShowingContentVersePath(bool bInShowingContentVersePath);
+	UE_API UE_INTERNAL void UpdatePaths();
 
 	using FIsAssetIdentifierPassingSearchFilterCallback = TFunction<bool(const FAssetIdentifier&)>;
 	UE_DEPRECATED(5.7, "Call SetDoesAssetPassSearchFilterCallback instead.")
-	void SetIsAssetIdentifierPassingSearchFilterCallback(const TOptional<FIsAssetIdentifierPassingSearchFilterCallback>& InIsAssetIdentifierPassingSearchFilterCallback);
+	UE_API void SetIsAssetIdentifierPassingSearchFilterCallback(const TOptional<FIsAssetIdentifierPassingSearchFilterCallback>& InIsAssetIdentifierPassingSearchFilterCallback);
 
 	using FDoesAssetPassSearchFilterCallback = TFunction<bool(const FAssetIdentifier&, const FAssetData&)>;
 	void SetDoesAssetPassSearchFilterCallback(const FDoesAssetPassSearchFilterCallback& InDoesAssetPassSearchFilterCallback) { DoesAssetPassSearchFilterCallback = InDoesAssetPassSearchFilterCallback; }
 
 	UE_DEPRECATED(5.6, "Use the ICollectionContainer overload instead.")
-	FName GetCurrentCollectionFilter() const;
-	void GetCurrentCollectionFilter(ICollectionContainer*& OutCollectionContainer, FName& OutCollectionName) const;
+	UE_API FName GetCurrentCollectionFilter() const;
+	UE_API void GetCurrentCollectionFilter(ICollectionContainer*& OutCollectionContainer, FName& OutCollectionName) const;
 
 	UE_DEPRECATED(5.6, "Use the ICollectionContainer overload instead.")
-	void SetCurrentCollectionFilter(FName NewFilter);
-	void SetCurrentCollectionFilter(const TSharedPtr<ICollectionContainer>& CollectionContainer, FName CollectionName);
+	UE_API void SetCurrentCollectionFilter(FName NewFilter);
+	UE_API void SetCurrentCollectionFilter(const TSharedPtr<ICollectionContainer>& CollectionContainer, FName CollectionName);
 
-	TArray<FName> GetCurrentPluginFilter() const;
-	void SetCurrentPluginFilter(TArray<FName> NewFilter);
-	TArray<FName> GetEncounteredPluginsAmongNodes() const;
+	UE_API TArray<FName> GetCurrentPluginFilter() const;
+	UE_API void SetCurrentPluginFilter(TArray<FName> NewFilter);
+	UE_API TArray<FName> GetEncounteredPluginsAmongNodes() const;
 
 	/* Delegate type to notify when the assets or NodeInfos have changed as opposed to when the filters changed */
 	FSimpleMulticastDelegate& OnAssetsChanged() { return OnAssetsChangedDelegate; }
 
 	/* Not to be confused with the above Content Browser Collection name, this is a TFiltercollection, a list of active filters */
-	void SetCurrentFilterCollection(TSharedPtr< TFilterCollection<FReferenceNodeInfo&> > NewFilterCollection);
+	UE_API void SetCurrentFilterCollection(TSharedPtr< TFilterCollection<FReferenceNodeInfo&> > NewFilterCollection);
 
 	/* Returns a set of unique asset types as UClass* */
 	const TSet<FTopLevelAssetPath>& GetAssetTypes() const { return CurrentClasses; }
@@ -132,10 +134,10 @@ public:
 	bool BreadthLimitExceeded() const { return bBreadthLimitReached; };
 
 	/** Refreshes the information of existing Referenced Properties Nodes */
-	void RefreshReferencedPropertiesNodes();
+	UE_API void RefreshReferencedPropertiesNodes();
 
 	/** Closes (removes) the specified Referenced Properties Node */
-	void CloseReferencedPropertiesNode(UEdGraphNode_ReferencedProperties* InNode);
+	UE_API void CloseReferencedPropertiesNode(UEdGraphNode_ReferencedProperties* InNode);
 
 private:
 
@@ -145,33 +147,33 @@ private:
 	 * @param InReferencedAsset: the asset referenced by the specified Object
 	 * @return the list of properties of the input Referencer, referencing the specified Asset.
 	 */
-	TArray<FReferencingPropertyDescription> RetrieveReferencingProperties(UObject* InReferencer, UObject* InReferencedAsset);
+	UE_API TArray<FReferencingPropertyDescription> RetrieveReferencingProperties(UObject* InReferencer, UObject* InReferencedAsset);
 
-	TWeakPtr<SReferenceViewer> GetReferenceViewer() const;
-	void SetReferenceViewer(TSharedPtr<SReferenceViewer> InViewer);
-	UEdGraphNode_Reference* ConstructNodes(const TArray<FAssetIdentifier>& GraphRootIdentifiers, const FIntPoint& GraphRootOrigin);
+	UE_API TWeakPtr<SReferenceViewer> GetReferenceViewer() const;
+	UE_API void SetReferenceViewer(TSharedPtr<SReferenceViewer> InViewer);
+	UE_API UEdGraphNode_Reference* ConstructNodes(const TArray<FAssetIdentifier>& GraphRootIdentifiers, const FIntPoint& GraphRootOrigin);
 
 
-	bool ExceedsMaxSearchDepth(int32 Depth, int32 MaxDepth) const;
-	bool ExceedsMaxSearchBreadth(int32 Breadth) const;
-	FAssetManagerDependencyQuery GetReferenceSearchFlags(bool bHardOnly) const;
+	UE_API bool ExceedsMaxSearchDepth(int32 Depth, int32 MaxDepth) const;
+	UE_API bool ExceedsMaxSearchBreadth(int32 Breadth) const;
+	UE_API FAssetManagerDependencyQuery GetReferenceSearchFlags(bool bHardOnly) const;
 
-	UEdGraphNode_Reference* CreateReferenceNode();
+	UE_API UEdGraphNode_Reference* CreateReferenceNode();
 
-	UEdGraphNode_ReferencedProperties* CreateReferencedPropertiesNode(const TArray<FReferencingPropertyDescription>& InPropertiesDescriptionArray
+	UE_API UEdGraphNode_ReferencedProperties* CreateReferencedPropertiesNode(const TArray<FReferencingPropertyDescription>& InPropertiesDescriptionArray
 	, const TObjectPtr<UEdGraphNode_Reference>& InReferencingNode, const TObjectPtr<UEdGraphNode_Reference>& InReferencedNode);
 
 	/* Generates a NodeInfo structure then used to generate and layout the graph nodes */
-	void RecursivelyPopulateNodeInfos(bool bReferencers, const TArray<FAssetIdentifier>& Identifiers, TMap<FAssetIdentifier, FReferenceNodeInfo>& NodeInfos, int32 CurrentDepth, int32 MaxDepth);
+	UE_API void RecursivelyPopulateNodeInfos(bool bReferencers, const TArray<FAssetIdentifier>& Identifiers, TMap<FAssetIdentifier, FReferenceNodeInfo>& NodeInfos, int32 CurrentDepth, int32 MaxDepth);
 
 	/* Marks up the NodeInfos with updated filter information and provision sizes */
-	void RecursivelyFilterNodeInfos(const FAssetIdentifier& InAssetId, TMap<FAssetIdentifier, FReferenceNodeInfo>& NodeInfos, int32 CurrentDepth, int32 MaxDepth);
+	UE_API void RecursivelyFilterNodeInfos(const FAssetIdentifier& InAssetId, TMap<FAssetIdentifier, FReferenceNodeInfo>& NodeInfos, int32 CurrentDepth, int32 MaxDepth);
 
 	/* Searches for the AssetData for the list of packages derived from the AssetReferences  */
-	void GatherAssetData(TMap<FAssetIdentifier, FReferenceNodeInfo>& InNodeInfos);
+	UE_API void GatherAssetData(TMap<FAssetIdentifier, FReferenceNodeInfo>& InNodeInfos);
 
 	/* Uses the NodeInfos map to generate and layout the graph nodes */
-	UEdGraphNode_Reference* RecursivelyCreateNodes(
+	UE_API UEdGraphNode_Reference* RecursivelyCreateNodes(
 		bool bInReferencers, 
 		const FAssetIdentifier& InAssetId, 
 		const FIntPoint& InNodeLoc, 
@@ -183,29 +185,29 @@ private:
 		bool bIsRoot = false
 	);
 
-	void ExpandNode(bool bReferencers, const FAssetIdentifier& InAssetIdentifier);
+	UE_API void ExpandNode(bool bReferencers, const FAssetIdentifier& InAssetIdentifier);
 
 	/** Removes all nodes from the graph */
-	void RemoveAllNodes();
+	UE_API void RemoveAllNodes();
 
 	/** Returns true if filtering is enabled and we have a valid collection */
-	bool ShouldFilterByCollection() const;
+	UE_API bool ShouldFilterByCollection() const;
 	
 	/** Returns true if filtering is enabled and we have a valid plugin name filter set */
-	bool ShouldFilterByPlugin() const;
+	UE_API bool ShouldFilterByPlugin() const;
 	
-	void GetUnfilteredGraphPluginNamesRecursive(bool bReferencers, const FAssetIdentifier& InAssetIdentifier, int32 InCurrentDepth, int32 InMaxDepth, const FAssetManagerDependencyQuery& Query, TSet<FAssetIdentifier>& OutAssetIdentifiers);
-	void GetUnfilteredGraphPluginNames(TArray<FAssetIdentifier> RootIdentifiers, TArray<FName>& OutPluginNames);
+	UE_API void GetUnfilteredGraphPluginNamesRecursive(bool bReferencers, const FAssetIdentifier& InAssetIdentifier, int32 InCurrentDepth, int32 InMaxDepth, const FAssetManagerDependencyQuery& Query, TSet<FAssetIdentifier>& OutAssetIdentifiers);
+	UE_API void GetUnfilteredGraphPluginNames(TArray<FAssetIdentifier> RootIdentifiers, TArray<FName>& OutPluginNames);
 
-	void GetSortedLinks(const TArray<FAssetIdentifier>& Identifiers, bool bReferencers, const FAssetManagerDependencyQuery& Query, TMap<FAssetIdentifier, EDependencyPinCategory>& OutLinks) const;
-	bool IsPackageIdentifierPassingFilter(const FAssetIdentifier& InAssetIdentifier) const;
-	bool IsPackageIdentifierPassingPluginFilter(const FAssetIdentifier& InAssetIdentifier) const;
-	bool DoesAssetPassSearchTextFilter(const FAssetIdentifier& InAssetIdentifier, const FAssetData& InAssetData) const;
+	UE_API void GetSortedLinks(const TArray<FAssetIdentifier>& Identifiers, bool bReferencers, const FAssetManagerDependencyQuery& Query, TMap<FAssetIdentifier, EDependencyPinCategory>& OutLinks) const;
+	UE_API bool IsPackageIdentifierPassingFilter(const FAssetIdentifier& InAssetIdentifier) const;
+	UE_API bool IsPackageIdentifierPassingPluginFilter(const FAssetIdentifier& InAssetIdentifier) const;
+	UE_API bool DoesAssetPassSearchTextFilter(const FAssetIdentifier& InAssetIdentifier, const FAssetData& InAssetData) const;
 
-	UEdGraphNode_Reference* FindPath(const FAssetIdentifier& RootId, const FAssetIdentifier& TargetId);
-	bool FindPath_Recursive(bool bInReferencers, const FAssetIdentifier& InAssetId, const FAssetIdentifier& Target, TMap<FAssetIdentifier, FReferenceNodeInfo>& InNodeInfos, TSet<FAssetIdentifier>& Visited);
+	UE_API UEdGraphNode_Reference* FindPath(const FAssetIdentifier& RootId, const FAssetIdentifier& TargetId);
+	UE_API bool FindPath_Recursive(bool bInReferencers, const FAssetIdentifier& InAssetId, const FAssetIdentifier& Target, TMap<FAssetIdentifier, FReferenceNodeInfo>& InNodeInfos, TSet<FAssetIdentifier>& Visited);
 
-	void RefreshReferencedPropertiesNode(const UEdGraphNode_ReferencedProperties* InNode);
+	UE_API void RefreshReferencedPropertiesNode(const UEdGraphNode_ReferencedProperties* InNode);
 private:
 	/** Pool for maintaining and rendering thumbnails */
 	TSharedPtr<FAssetThumbnailPool> AssetThumbnailPool;
@@ -260,3 +262,5 @@ private:
 	friend SReferenceViewer;
 	friend class UReferenceViewerSchema;
 };
+
+#undef UE_API
diff --git a/Engine/Plugins/Editor/AssetReferenceRestrictions/Source/AssetReferenceRestrictions/Public/AssetReferencingDomains.h b/Engine/Plugins/Editor/AssetReferenceRestrictions/Source/AssetReferenceRestrictions/Public/AssetReferencingDomains.h
--- a/Engine/Plugins/Editor/AssetReferenceRestrictions/Source/AssetReferenceRestrictions/Public/AssetReferencingDomains.h
+++ b/Engine/Plugins/Editor/AssetReferenceRestrictions/Source/AssetReferenceRestrictions/Public/AssetReferencingDomains.h
@@ -4,6 +4,8 @@
 
 #include "Internationalization/Text.h"
 
+#define UE_API ASSETREFERENCERESTRICTIONS_API
+
 //@TODO: Debug printing
 #define UE_ASSET_DOMAIN_FILTERING_DEBUG_LOGGING 0
 
@@ -51,40 +53,40 @@ struct FDomainData final : public TSharedFromThis<FDomainData>
 	}
 };
 
-struct ASSETREFERENCERESTRICTIONS_API FDomainDatabase final
+struct FDomainDatabase final
 {
-	FDomainDatabase();
-	~FDomainDatabase();
+	UE_API FDomainDatabase();
+	UE_API ~FDomainDatabase();
 
-	void Init();
+	UE_API void Init();
 
-	void MarkDirty();
-	void UpdateIfNecessary();
+	UE_API void MarkDirty();
+	UE_API void UpdateIfNecessary();
 
 	// Delegate that's called whenever the database has been updated to allow people to respond to the change
 	TMulticastDelegate<void()>& PostDatabaseUpdated() { return OnPostDatabaseUpdated; };
 
-	void ValidateAllDomains();
-	void DebugPrintAllDomains();
+	UE_API void ValidateAllDomains();
+	UE_API void DebugPrintAllDomains();
 
-	void OnPluginCreatedOrMounted(IPlugin& NewPlugin);
+	UE_API void OnPluginCreatedOrMounted(IPlugin& NewPlugin);
 
-	TSharedPtr<FDomainData> FindDomainFromAssetData(const FAssetData& AssetData) const;
+	UE_API TSharedPtr<FDomainData> FindDomainFromAssetData(const FAssetData& AssetData) const;
 
-	TTuple<bool, FText> CanDomainsSeeEachOther(TSharedPtr<FDomainData> Referencee, TSharedPtr<FDomainData> Referencer) const;
+	UE_API TTuple<bool, FText> CanDomainsSeeEachOther(TSharedPtr<FDomainData> Referencee, TSharedPtr<FDomainData> Referencer) const;
 
 	const TArray<FString>& GetDomainsDefinedByPlugins() const { return DomainsDefinedByPlugins; }
 	bool IsPlugin(TSharedPtr<FDomainData> Domain) const { return PluginDomains.Contains(Domain); }
 
-	TSharedPtr<FDomainData> FindOrAddDomainByName(const FString& Name);
+	UE_API TSharedPtr<FDomainData> FindOrAddDomainByName(const FString& Name);
 
 private:
-	void RebuildFromScratch();
+	UE_API void RebuildFromScratch();
 	
-	void BuildDomainFromPlugin(TSharedRef<IPlugin> Plugin);
-	void BuildUnrestrictedDomainFromPlugin(TSharedRef<IPlugin> Plugin);
+	UE_API void BuildDomainFromPlugin(TSharedRef<IPlugin> Plugin);
+	UE_API void BuildUnrestrictedDomainFromPlugin(TSharedRef<IPlugin> Plugin);
 
-	void AddDomainVisibilityList(TSharedPtr<FDomainData> Domain, const TArray<FString>& VisibilityList);
+	UE_API void AddDomainVisibilityList(TSharedPtr<FDomainData> Domain, const TArray<FString>& VisibilityList);
 private:
 	// Map from domain name to 
 	TMap<FString, TSharedPtr<FDomainData>> DomainNameMap;
@@ -121,3 +123,5 @@ private:
 
 	TMulticastDelegate<void()> OnPostDatabaseUpdated;
 };
+
+#undef UE_API
diff --git a/Engine/Plugins/Editor/AssetReferenceRestrictions/Source/AssetReferenceRestrictions/Public/AssetReferencingPolicySettings.h b/Engine/Plugins/Editor/AssetReferenceRestrictions/Source/AssetReferenceRestrictions/Public/AssetReferencingPolicySettings.h
--- a/Engine/Plugins/Editor/AssetReferenceRestrictions/Source/AssetReferenceRestrictions/Public/AssetReferencingPolicySettings.h
+++ b/Engine/Plugins/Editor/AssetReferenceRestrictions/Source/AssetReferenceRestrictions/Public/AssetReferencingPolicySettings.h
@@ -6,6 +6,8 @@
 
 #include "AssetReferencingPolicySettings.generated.h"
 
+#define UE_API ASSETREFERENCERESTRICTIONS_API
+
 struct FDirectoryPath;
 struct FPropertyChangedEvent;
 
@@ -174,8 +176,8 @@ struct FARPDomainDefinitionByContentRoot
 };
 
 /** Settings for the Asset Referencing Policy, these settings are used to determine which plugins and game folders can reference content from each other */
-UCLASS(config=Game, defaultconfig, meta=(DisplayName="Asset Referencing Policy"))
-class ASSETREFERENCERESTRICTIONS_API UAssetReferencingPolicySettings : public UDeveloperSettings
+UCLASS(MinimalAPI, config=Game, defaultconfig, meta=(DisplayName="Asset Referencing Policy"))
+class UAssetReferencingPolicySettings : public UDeveloperSettings
 {
 	GENERATED_BODY()
 
@@ -201,23 +203,23 @@ public:
 	bool bIgnoreEditorOnlyReferences = true;
 
 	// The names of the project, special system mount, and game domains
-	static const FString EngineDomainName;
-	static const FString ScriptDomainName;
-	static const FString GameDomainName;
-	static const FString TempDomainName;
-	static const FString NeverCookDomainName;
+	static UE_API const FString EngineDomainName;
+	static UE_API const FString ScriptDomainName;
+	static UE_API const FString GameDomainName;
+	static UE_API const FString TempDomainName;
+	static UE_API const FString NeverCookDomainName;
 
 public:
 	//~UObject interface
 #if WITH_EDITOR
-	virtual void PostEditChangeProperty(FPropertyChangedEvent& PropertyChangedEvent) override;
+	UE_API virtual void PostEditChangeProperty(FPropertyChangedEvent& PropertyChangedEvent) override;
 #endif
 	//~End of UObject interface
 
-	UAssetReferencingPolicySettings(const FObjectInitializer& ObjectInitializer);
+	UE_API UAssetReferencingPolicySettings(const FObjectInitializer& ObjectInitializer);
 
 private:
-	TArray<FString> GetListOfDomains(bool bAllowPlugins, bool bAllowEngine, bool bAllowGame) const;
+	UE_API TArray<FString> GetListOfDomains(bool bAllowPlugins, bool bAllowEngine, bool bAllowGame) const;
 
 	UFUNCTION()
 	TArray<FString> GetListOfDomains_All() const
@@ -244,3 +246,5 @@ private:
 	}
 
 };
+
+#undef UE_API
diff --git a/Engine/Plugins/Editor/AssetReferenceRestrictions/Source/AssetReferenceRestrictions/Public/AssetReferencingPolicySubsystem.h b/Engine/Plugins/Editor/AssetReferenceRestrictions/Source/AssetReferenceRestrictions/Public/AssetReferencingPolicySubsystem.h
--- a/Engine/Plugins/Editor/AssetReferenceRestrictions/Source/AssetReferenceRestrictions/Public/AssetReferencingPolicySubsystem.h
+++ b/Engine/Plugins/Editor/AssetReferenceRestrictions/Source/AssetReferenceRestrictions/Public/AssetReferencingPolicySubsystem.h
@@ -8,6 +8,8 @@
 
 #include "AssetReferencingPolicySubsystem.generated.h"
 
+#define UE_API ASSETREFERENCERESTRICTIONS_API
+
 enum class EAssetReferenceFilterRole : uint8;
 struct FAssetReferenceFilterContext;
 class IAssetReferenceFilter;
@@ -29,30 +31,32 @@ struct FAssetReferenceError
 };
 
 /** Subsystem to register the domain-based asset referencing policy restrictions with the editor */
-UCLASS()
-class ASSETREFERENCERESTRICTIONS_API UAssetReferencingPolicySubsystem : public UEditorSubsystem
+UCLASS(MinimalAPI)
+class UAssetReferencingPolicySubsystem : public UEditorSubsystem
 {
 	GENERATED_BODY()
 
 public:
 	//~UEditorSubsystem interface
-	virtual bool ShouldCreateSubsystem(UObject* Outer) const override;
-	virtual void Initialize(FSubsystemCollectionBase& Collection) override;
-	virtual void Deinitialize() override;
+	UE_API virtual bool ShouldCreateSubsystem(UObject* Outer) const override;
+	UE_API virtual void Initialize(FSubsystemCollectionBase& Collection) override;
+	UE_API virtual void Deinitialize() override;
 	//~End of UEditorSubsystem interface
 
 	// Returns whether the given asset's outgoing references are restricted in any way and should be individually validated
-	bool ShouldValidateAssetReferences(const FAssetData& Asset) const;
+	UE_API bool ShouldValidateAssetReferences(const FAssetData& Asset) const;
 
 	// Check the outgoing references of the given asset according to the asset registry and return details of any errors 
-	TValueOrError<void, TArray<FAssetReferenceError>> ValidateAssetReferences(const FAssetData& Asset) const;
-	TValueOrError<void, TArray<FAssetReferenceError>> ValidateAssetReferences(const FAssetData& Asset, const EAssetReferenceFilterRole Role) const;
+	UE_API TValueOrError<void, TArray<FAssetReferenceError>> ValidateAssetReferences(const FAssetData& Asset) const;
+	UE_API TValueOrError<void, TArray<FAssetReferenceError>> ValidateAssetReferences(const FAssetData& Asset, const EAssetReferenceFilterRole Role) const;
 
-	TSharedPtr<FDomainDatabase> GetDomainDB() const;
+	UE_API TSharedPtr<FDomainDatabase> GetDomainDB() const;
 private:
-	void UpdateDBIfNecessary() const;
+	UE_API void UpdateDBIfNecessary() const;
 
-	TSharedPtr<IAssetReferenceFilter> HandleMakeAssetReferenceFilter(const FAssetReferenceFilterContext& Context);
+	UE_API TSharedPtr<IAssetReferenceFilter> HandleMakeAssetReferenceFilter(const FAssetReferenceFilterContext& Context);
 
 	TSharedPtr<FDomainDatabase> DomainDB;
 };
+
+#undef UE_API
diff --git a/Engine/Plugins/Editor/AssetReferenceRestrictions/Source/AssetReferenceRestrictions/Public/DomainAssetReferenceFilter.h b/Engine/Plugins/Editor/AssetReferenceRestrictions/Source/AssetReferenceRestrictions/Public/DomainAssetReferenceFilter.h
--- a/Engine/Plugins/Editor/AssetReferenceRestrictions/Source/AssetReferenceRestrictions/Public/DomainAssetReferenceFilter.h
+++ b/Engine/Plugins/Editor/AssetReferenceRestrictions/Source/AssetReferenceRestrictions/Public/DomainAssetReferenceFilter.h
@@ -6,35 +6,37 @@
 #include "Editor/EditorEngine.h"
 #include "Editor/AssetReferenceFilter.h"
 
+#define UE_API ASSETREFERENCERESTRICTIONS_API
+
 struct FAssetData;
 struct FDomainDatabase;
 struct FDomainData;
 
-class ASSETREFERENCERESTRICTIONS_API FDomainAssetReferenceFilter : public IAssetReferenceFilter
+class FDomainAssetReferenceFilter : public IAssetReferenceFilter
 {
 public:
-	FDomainAssetReferenceFilter(const FAssetReferenceFilterContext& Context, TSharedPtr<FDomainDatabase> InDomainDB);
-	~FDomainAssetReferenceFilter();
+	UE_API FDomainAssetReferenceFilter(const FAssetReferenceFilterContext& Context, TSharedPtr<FDomainDatabase> InDomainDB);
+	UE_API ~FDomainAssetReferenceFilter();
 
 	//~IAssetReferenceFilter interface
-	virtual bool PassesFilter(const FAssetData& AssetData, FText* OutOptionalFailureReason = nullptr) const override;
+	UE_API virtual bool PassesFilter(const FAssetData& AssetData, FText* OutOptionalFailureReason = nullptr) const override;
 	//~End of IAssetReferenceFilter
 
 	// Update any cached information for all filters
-	static void UpdateAllFilters();
+	static UE_API void UpdateAllFilters();
 private:
 	using FAssetDataInfo = TTuple<FAssetData, TSharedPtr<FDomainData>>;
 
-	bool PassesFilterImpl(const FAssetData& AssetData, FText& OutOptionalFailureReason) const;
-	bool IsCrossPluginReferenceAllowed(const FAssetDataInfo& ReferencingAssetDataInfo, const FAssetDataInfo& ReferencedAssetDataInfo) const;
+	UE_API bool PassesFilterImpl(const FAssetData& AssetData, FText& OutOptionalFailureReason) const;
+	UE_API bool IsCrossPluginReferenceAllowed(const FAssetDataInfo& ReferencingAssetDataInfo, const FAssetDataInfo& ReferencedAssetDataInfo) const;
 
-	void DetermineReferencingDomain();
+	UE_API void DetermineReferencingDomain();
 
 	/** Heuristic to find actual assets from preview assets (i.e., the material editor's preview material) */
-	void TryGetAssociatedAssetsFromPossiblyPreviewObject(UObject* PossiblyPreviewObject, TArray<FAssetData>& InOutAssetsToConsider) const;
+	UE_API void TryGetAssociatedAssetsFromPossiblyPreviewObject(UObject* PossiblyPreviewObject, TArray<FAssetData>& InOutAssetsToConsider) const;
 
 private:
-	static TArray<FDomainAssetReferenceFilter*> FilterInstances;
+	static UE_API TArray<FDomainAssetReferenceFilter*> FilterInstances;
 
 	TSharedPtr<FDomainDatabase> DomainDB;
 
@@ -43,3 +45,5 @@ private:
 
 	FText Failure_CouldNotDetermineDomain;
 };
+
+#undef UE_API
