- Exclude unmounted plugins from FPluginUtils::UnloadPluginAssets
- In FPluginUtils::UnloadPlugins, change IPlugin::IsEnabled check with IPlugin::IsMounted; usually plugins are enabled and mounted at the same time, with the exception of 'ExplicitlyLoaded' plugins which can be enabled but not mounted until much later

#rnx

diff --git a/Engine/Plugins/Developer/PluginUtils/Source/PluginUtils/Private/PluginUtils.cpp b/Engine/Plugins/Developer/PluginUtils/Source/PluginUtils/Private/PluginUtils.cpp
--- a/Engine/Plugins/Developer/PluginUtils/Source/PluginUtils/Private/PluginUtils.cpp
+++ b/Engine/Plugins/Developer/PluginUtils/Source/PluginUtils/Private/PluginUtils.cpp
@@ -1076,7 +1076,7 @@ bool FPluginUtils::UnloadPlugins(const TConstArrayView<TSharedRef<IPlugin>> Plug
 
 		for (const TSharedRef<IPlugin>& Plugin : Plugins)
 		{
-			if (Plugin->IsEnabled())
+			if (Plugin->IsMounted())
 			{
 				bPluginUnmounted = true;
 
@@ -1136,7 +1136,10 @@ bool FPluginUtils::UnloadPluginsAssets(const TConstArrayView<TSharedRef<IPlugin>
 	PluginNames.Reserve(Plugins.Num());
 	for (const TSharedRef<IPlugin>& Plugin : Plugins)
 	{
-		PluginNames.Add(Plugin->GetName());
+		if (Plugin->IsMounted())
+		{
+			PluginNames.Add(Plugin->GetName());
+		}
 	}
 
 	return UnloadPluginsAssets(PluginNames, OutFailReason);
