Add gameplaytag fixup to plugin creation from template

#rnx

diff --git a/Engine/Plugins/Developer/PluginUtils/PluginUtils.uplugin b/Engine/Plugins/Developer/PluginUtils/PluginUtils.uplugin
--- a/Engine/Plugins/Developer/PluginUtils/PluginUtils.uplugin
+++ b/Engine/Plugins/Developer/PluginUtils/PluginUtils.uplugin
@@ -20,5 +20,11 @@
 			"Type": "Editor",
 			"LoadingPhase": "Default"
 		}
+	],
+	"Plugins": [
+		{
+			"Name": "EditorScriptingUtilities",
+			"Enabled": true
+		}
 	]
 }
diff --git a/Engine/Plugins/Developer/PluginUtils/Source/PluginUtils/PluginUtils.Build.cs b/Engine/Plugins/Developer/PluginUtils/Source/PluginUtils/PluginUtils.Build.cs
--- a/Engine/Plugins/Developer/PluginUtils/Source/PluginUtils/PluginUtils.Build.cs
+++ b/Engine/Plugins/Developer/PluginUtils/Source/PluginUtils/PluginUtils.Build.cs
@@ -15,13 +15,17 @@ namespace UnrealBuildTool.Rules
 
 			PrivateDependencyModuleNames.AddRange(
 				new string[] {
-					"UnrealEd",
-					"GameProjectGeneration",
-					"DesktopPlatform",
 					"AssetRegistry",
 					"AssetTools",
 					"CoreUObject",
-					"SourceControl"
+					"DesktopPlatform",
+					"EditorScriptingUtilities",
+					"Engine",
+					"GameProjectGeneration",
+					"GameplayTags",
+					"SourceControl",
+					"SubobjectDataInterface",
+					"UnrealEd"
 				}
 			);
 		}
diff --git a/Engine/Plugins/Developer/PluginUtils/Source/PluginUtils/Private/PluginUtils.cpp b/Engine/Plugins/Developer/PluginUtils/Source/PluginUtils/Private/PluginUtils.cpp
--- a/Engine/Plugins/Developer/PluginUtils/Source/PluginUtils/Private/PluginUtils.cpp
+++ b/Engine/Plugins/Developer/PluginUtils/Source/PluginUtils/Private/PluginUtils.cpp
@@ -3,28 +3,34 @@
 #if WITH_EDITOR
 
 #include "PluginUtils.h"
-#include "IDesktopPlatform.h"
-#include "SourceControlHelpers.h"
-#include "IContentBrowserSingleton.h"
-#include "ContentBrowserModule.h"
-#include "GameProjectUtils.h"
-#include "Interfaces/IProjectManager.h"
-#include "Interfaces/IPluginManager.h"
+
 #include "AssetRegistry/AssetRegistryModule.h"
-#include "IAssetTools.h"
 #include "AssetToolsModule.h"
+#include "ContentBrowserModule.h"
 #include "DesktopPlatformModule.h"
-#include "LocalizationDescriptor.h"
-#include "PackageTools.h"
+#include "EditorAssetLibrary.h"
+#include "GameProjectUtils.h"
+#include "GameplayTagContainer.h"
+#include "GameplayTagsManager.h"
 #include "HAL/PlatformFileManager.h"
+#include "IAssetTools.h"
+#include "IContentBrowserSingleton.h"
+#include "IDesktopPlatform.h"
+#include "Interfaces/IPluginManager.h"
+#include "Interfaces/IProjectManager.h"
 #include "Internationalization/TextPackageNamespaceUtil.h"
-#include "Misc/Paths.h"
-#include "Misc/PathViews.h"
+#include "LocalizationDescriptor.h"
+#include "Misc/FeedbackContext.h"
 #include "Misc/FileHelper.h"
 #include "Misc/PackageName.h"
-#include "Misc/FeedbackContext.h"
+#include "Misc/PathViews.h"
+#include "Misc/Paths.h"
+#include "PackageTools.h"
 #include "PluginReferenceDescriptor.h"
+#include "SourceControlHelpers.h"
+#include "SubobjectDataSubsystem.h"
 #include "UObject/CoreRedirects.h"
+#include "UObject/PropertyAccessUtil.h"
 #include "UObject/UObjectIterator.h"
 
 DEFINE_LOG_CATEGORY_STATIC(LogPluginUtils, Log, All);
@@ -191,7 +197,8 @@ namespace PluginUtils
 						}
 
 						// If file of specified extension - open the file as text and replace PLUGIN_NAME in there before saving
-						if (NameReplacementFileTypes.Contains(NewExt))
+						const bool bIsTagsIni = NewExt.Equals(TEXT("ini")) && NewName.Contains(TEXT("/Config/Tags/"));
+						if (NameReplacementFileTypes.Contains(NewExt) || bIsTagsIni)
 						{
 							FString OutFileContents;
 							if (!FFileHelper::LoadFileToString(OutFileContents, FilenameOrDirectory))
@@ -293,6 +300,35 @@ namespace PluginUtils
 				return true;
 			}
 
+			void ReplaceGameplayTagsInObject(UObject* Object)
+			{
+				if (!Object)
+				{
+					return;
+				}
+
+				for (TPropertyValueIterator<FProperty> It(Object->GetClass(), Object, EPropertyValueIteratorFlags::FullRecursion); It; ++It)
+				{
+					if (const FStructProperty* StructProperty = CastField<FStructProperty>(It.Key()))
+					{
+						if (UScriptStruct* StructType = StructProperty->Struct)
+						{
+							if (StructType->IsChildOf(FGameplayTag::StaticStruct()))
+							{
+								const void* PropertyValue = It.Value();
+
+								const FGameplayTag* GameplayTag = reinterpret_cast<const FGameplayTag*>(PropertyValue);
+								if (GameplayTag && GameplayTag->IsValid() && GameplayTag->ToString().Contains(PLUGIN_NAME))
+								{
+									FGameplayTag NewTag = UGameplayTagsManager::Get().RequestGameplayTag(FName(GameplayTag->ToString().Replace(*PLUGIN_NAME, *NameToReplace, ESearchCase::CaseSensitive)));
+									It.Key()->CopyCompleteValue(const_cast<void*>(PropertyValue), &NewTag);
+								}
+							}
+						}
+					}
+				}
+			}
+
 			/**
 			 * Fixes up any assets that contain the PLUGIN_NAME text macro, since those need to be renamed by the engine for the change to
 			 * stick (as opposed to just renaming the file) as well as regenerating the localization ids for the assets in the plugin.
@@ -319,6 +355,34 @@ namespace PluginUtils
 
 						for (FAssetData AssetData : Assets)
 						{
+							UObject* LoadedAsset = AssetData.GetAsset();
+							const UClass* LoadedClass = AssetData.GetClass(EResolveClass::Yes);
+							if (LoadedClass == UBlueprint::StaticClass())
+							{
+								if (UClass* BPClass = UEditorAssetLibrary::LoadBlueprintClass(AssetData.GetPackage()->GetPathName()))
+								{
+									LoadedAsset = BPClass->GetDefaultObject();
+
+									USubobjectDataSubsystem* DataSubsystem = USubobjectDataSubsystem::Get();
+
+									if (ensure(DataSubsystem))
+									{
+										TArray<FSubobjectDataHandle> SubobjectData;
+										DataSubsystem->GatherSubobjectData(LoadedAsset, SubobjectData);
+
+										for (const FSubobjectDataHandle& Handle : SubobjectData)
+										{
+											if (UObject* SubObject = const_cast<UObject*>(Handle.GetData()->GetObjectForBlueprint(Cast<UBlueprint>(LoadedAsset))))
+											{
+												ReplaceGameplayTagsInObject(SubObject);
+											}
+										}
+									}
+								}
+							}
+
+							ReplaceGameplayTagsInObject(LoadedAsset);
+
 							const FString AssetName = AssetData.AssetName.ToString().Replace(*PLUGIN_NAME, *NameToReplace, ESearchCase::CaseSensitive);
 							const FString AssetPath = AssetData.PackagePath.ToString().Replace(*PLUGIN_NAME, *NameToReplace, ESearchCase::CaseSensitive);
 
@@ -372,6 +436,9 @@ namespace PluginUtils
 
 		if (Plugin.IsValid())
 		{
+			UGameplayTagsManager::Get().AddTagIniSearchPath(Plugin->GetBaseDir() / TEXT("Config") / TEXT("Tags"));
+			UGameplayTagsManager::Get().EditorRefreshGameplayTagTree();
+
 			const FString PluginBaseDir = Plugin->GetBaseDir();
 			FFixupPluginAssets FixupPluginAssets(PlatformFile, NameToReplace, PluginBaseDir);
 			PlatformFile.IterateDirectoryRecursively(*PluginBaseDir, FixupPluginAssets);
