Fixed duplicate package localization ID warning after plugin creation via plugin templates

diff --git a/Engine/Plugins/Developer/PluginUtils/Source/PluginUtils/Private/PluginUtils.cpp b/Engine/Plugins/Developer/PluginUtils/Source/PluginUtils/Private/PluginUtils.cpp
--- a/Engine/Plugins/Developer/PluginUtils/Source/PluginUtils/Private/PluginUtils.cpp
+++ b/Engine/Plugins/Developer/PluginUtils/Source/PluginUtils/Private/PluginUtils.cpp
@@ -394,13 +394,14 @@ namespace PluginUtils
 								FAssetRenameData RenameData(AssetData.GetAsset(), AssetPath, AssetName);
 
 								AssetRenameData.Add(RenameData);
+							}
 
-								if (UObject* Asset = AssetData.GetAsset())
-								{
-									Asset->Modify();
-									TextNamespaceUtil::ClearPackageNamespace(Asset);
-									TextNamespaceUtil::EnsurePackageNamespace(Asset);
-								}
+							// ANY assets created/duplicated from a plugin template should have its package localization ID reset to avoid conflicts.
+							if (UObject* Asset = AssetData.GetAsset())
+							{
+								Asset->Modify();
+								TextNamespaceUtil::ClearPackageNamespace(Asset);
+								TextNamespaceUtil::EnsurePackageNamespace(Asset);
 							}
 						}
 					}
