Plugin Wizard can now run post-create python scripts on plugin template creation

[RN] minor plugin

#Synced-CL 48656766

diff --git a/Engine/Plugins/Developer/PluginUtils/Source/PluginUtils/Private/PluginUtils.cpp b/Engine/Plugins/Developer/PluginUtils/Source/PluginUtils/Private/PluginUtils.cpp
--- a/Engine/Plugins/Developer/PluginUtils/Source/PluginUtils/Private/PluginUtils.cpp
+++ b/Engine/Plugins/Developer/PluginUtils/Source/PluginUtils/Private/PluginUtils.cpp
@@ -29,6 +29,8 @@
 #include "PluginReferenceDescriptor.h"
 #include "SourceControlHelpers.h"
 #include "SubobjectDataSubsystem.h"
+#include "UnrealEdMisc.h"
+#include "Misc/CommandLine.h"
 #include "UObject/CoreRedirects.h"
 #include "UObject/PropertyAccessUtil.h"
 #include "UObject/UObjectIterator.h"
@@ -554,6 +556,10 @@ namespace PluginUtils
 	{
 		ExCreationParams.PluginIconPath = CreationParams.PluginIconPath;
 		ExCreationParams.TemplateFolders = CreationParams.TemplateFolders;
+		
+		ExCreationParams.PostCreatePythonScriptPath            = CreationParams.PostCreatePythonScriptPath;
+		ExCreationParams.PostCreatePythonScriptArguments       = CreationParams.PostCreatePythonScriptArguments;
+		ExCreationParams.PostCreatePythonScriptAdditionalArguments = CreationParams.PostCreatePythonScriptAdditionalArguments;
 
 		ExCreationParams.Descriptor.FriendlyName = CreationParams.FriendlyName.Len() > 0 ? CreationParams.FriendlyName : PluginName;
 		ExCreationParams.Descriptor.Version = 1;
@@ -578,6 +584,180 @@ namespace PluginUtils
 			ExCreationParams.Descriptor.Modules.Add(FModuleDescriptor(*PluginName, CreationParams.ModuleDescriptorType, CreationParams.LoadingPhase));
 		}
 	}
+	
+	FString GetTemplatePluginNames(const TArray<FString> TemplateFolders)
+    {
+    	FString PluginNames;
+    	for (const FString& TemplateFolder : TemplateFolders)
+    	{
+        	TArray<FString> UPluginFiles;
+        	IFileManager::Get().FindFiles(UPluginFiles, *TemplateFolder, TEXT("*.uplugin"));
+    
+    		for (const FString& UPluginFile : UPluginFiles)
+    		{
+				if (!PluginNames.IsEmpty())
+				{
+					PluginNames += ",";
+				}
+				// Strip the .uplugin extension from each file
+                PluginNames += FPaths::GetBaseFilename(UPluginFile);
+			}
+    	}
+        
+        return PluginNames;
+    }
+	
+	bool RunPostCreatePythonScript(const FString& PluginName, const FString& PluginFolder, const TArray<FString> TemplateFolders, const FString& ScriptPath, const FString& TemplateArgs, const FString& AdditionalArgs, const bool IsEnginePlugin, FText& OutFailReason)
+	{
+		FString AbsScriptPath = IFileManager::Get().ConvertToAbsolutePathForExternalAppForRead(*ScriptPath);
+		if (!FPaths::FileExists(AbsScriptPath))
+		{
+			OutFailReason = FText::Format(LOCTEXT("PostCreateScriptNotFound", "Post Create Python Script not found in: {0}"), FText::FromString(AbsScriptPath));
+			return false;
+		}
+		
+		const TSharedPtr<IPlugin> PythonPlugin = IPluginManager::Get().FindPlugin(TEXT("PythonScriptPlugin"));
+		if (!PythonPlugin.IsValid() || !PythonPlugin->IsEnabled())
+		{
+			OutFailReason = LOCTEXT("PostCreatePythonNotAvailable", "PythonScriptPlugin not loaded and enabled. Post Create Python Script cannot run without it.");
+			return false;
+		}
+		
+		const FString TemplatePluginNames = GetTemplatePluginNames(TemplateFolders);
+		if (TemplatePluginNames.IsEmpty())
+		{
+			OutFailReason = LOCTEXT("PostCreateNoTemplatePluginName", "Could not find a valid Template Plugin Name.");
+			return false;
+		}
+		
+		const FString EditorExe = FUnrealEdMisc::Get().GetExecutableForCommandlets();
+		if (EditorExe.IsEmpty())
+		{
+			OutFailReason = LOCTEXT("PostCreateNoEditorExe", "Could not find commandlet executable for running the post-create Python script.");
+			return false;
+		}
+		
+		void* ReadPipe = nullptr, *WritePipe = nullptr;
+		if (!FPlatformProcess::CreatePipe(ReadPipe, WritePipe))
+		{
+			OutFailReason = LOCTEXT("PostCreatePipeFailed", "Failed to create pipe for post-create Python script.");
+			return false;
+		}
+
+		FString CommandLine = FCommandLine::GetOriginal();
+		
+		auto WrapInQuotes = [](const FString& In) -> FString
+			{
+				FString Out = In;
+				Out.InsertAt(0, "\"");
+				Out.Append("\"");
+				return Out;
+			};
+		
+		auto SwapToForwardSlash = [](const FString& In) -> FString
+			{
+				FString Out = In;
+				Out.ReplaceInline(TEXT("\\"), TEXT("/")); // swap backslash to forward slash for paths
+				return Out;
+			};
+		
+		auto EscapeForScript = [](const FString& In) -> FString
+			{
+				FString Out = In;
+				Out.ReplaceInline(TEXT("\'"), TEXT("\"")); // swap single for double quotes
+				Out.ReplaceInline(TEXT("\""), TEXT("\\\\\"\"\"")); // enough escapes to go through the commandlet + python command stripping 
+				return Out;
+			};
+		
+		FString TemplatePathArgs;
+		for (const FString& Folder : TemplateFolders)
+		{
+			TemplatePathArgs += FString::Printf(TEXT("--TemplatePath \"%s\" "), *SwapToForwardSlash(FPaths::ConvertRelativePathToFull(Folder)));
+		}
+		
+		CommandLine += FString::Printf(TEXT(" -stdout -FullStdOutLogOutput -redirectstderr -EnablePlugins=\"%s\" -run=pythonscript -script=\"%s --PluginName %s --PluginRoot %s --IsEnginePlugin %d %s %s %s\""), *TemplatePluginNames, *EscapeForScript(AbsScriptPath), *EscapeForScript(WrapInQuotes(PluginName)), *EscapeForScript(WrapInQuotes(SwapToForwardSlash(PluginFolder))), IsEnginePlugin ? 1 : 0, *EscapeForScript(TemplatePathArgs), *EscapeForScript(TemplateArgs), *EscapeForScript(AdditionalArgs));
+		
+		FScopedSlowTask Progress(/*InAmountOfWork=*/ 3.f, LOCTEXT("RunningPostCreatePython", "Running post-create Python script..."), /*bInEnabled=*/ true, *GWarn);
+		Progress.MakeDialog(/*bShowCancelButton=*/ false);
+		
+		FProcHandle ProcHandle = FPlatformProcess::CreateProc(*EditorExe, *CommandLine, /*bLaunchDetached*/ false, /*bLaunchHidden*/ true, /*bLaunchReallyHidden*/ true, /*OutProcessId*/ nullptr, /*PriorityModifier*/ 0, /*OptionalWorkingDirectory*/ nullptr, /*PipeWrite*/ WritePipe);
+		if (!ProcHandle.IsValid())
+		{
+			FPlatformProcess::ClosePipe(ReadPipe, WritePipe);
+			OutFailReason = LOCTEXT( "PostCreateProcFailed", "Failed to create a process for the post-create Python script.");
+			return false;
+		}
+		
+		UE_LOG(LogPluginUtils, Display, TEXT("Launching new process to run python post-create script..."));
+		
+		bool bSeenRunningPython = false;
+		bool bSeenScriptCompleted = false;
+		float ProgressThisFrame = 1.0f;
+		FText CurrentStatusText = LOCTEXT( "PostCreateLaunchProc", "Launching new process to run python post-create script...");
+		
+		int32 ReturnCode = -1;
+		while (FPlatformProcess::IsProcRunning(ProcHandle))
+		{
+			FPlatformProcess::Sleep(0.1f);
+			
+			const FString STDLog = FPlatformProcess::ReadPipe(ReadPipe);
+			if (!STDLog.IsEmpty())
+			{
+				// Log the process' output to our output log
+				UE_LOG(LogPluginUtils, Display, TEXT("%s"), *STDLog);
+				
+				// Update UI Progress bar based on log entries
+				TArray<FString> Lines;
+				STDLog.ParseIntoArrayLines(Lines);
+
+				for (const FString& Line : Lines)
+				{
+					if (!bSeenRunningPython &&	
+						Line.Contains(TEXT("LogPythonScriptCommandlet: Display: Running Python script")))
+					{
+						bSeenRunningPython = true;
+						ProgressThisFrame = 1.0f;
+						CurrentStatusText = LOCTEXT( "PostCreateRunningScript", "Running Python script in process...");
+					}
+					else if (!bSeenScriptCompleted &&
+						 (Line.Contains(TEXT("LogPythonScriptCommandlet: Display: Python script executed successfully")) ||
+						  Line.Contains(TEXT("Commandlet PythonScriptCommandlet_0 finished execution")) ||
+						  Line.Contains(TEXT("Success - 0 error(s), 0 warning(s)")) ||
+						  Line.Contains(TEXT("Execution of commandlet took"))))
+					{
+						bSeenScriptCompleted = true;
+						ProgressThisFrame = 1.0f;
+						CurrentStatusText = LOCTEXT( "PostCreateScriptFinished", "Python script finished. Closing process...");
+					}
+				}
+			}
+
+			// Update UI progress bar
+			Progress.EnterProgressFrame(ProgressThisFrame, CurrentStatusText);
+			ProgressThisFrame = 0;
+		}
+		
+		UE_LOG(LogPluginUtils, Display, TEXT("Python post-create script process completed."));
+		
+		FPlatformProcess::ClosePipe(ReadPipe, WritePipe);
+		
+		if (!FPlatformProcess::GetProcReturnCode(ProcHandle, &ReturnCode))
+		{
+			FPlatformProcess::CloseProc(ProcHandle);
+			OutFailReason = LOCTEXT("PostCreateFailedGetReturnCode", "Failed to get return code for post-create Python script process.");
+			return false;
+		}
+		
+		FPlatformProcess::CloseProc(ProcHandle);
+		
+		if (ReturnCode != 0)
+		{
+			OutFailReason = LOCTEXT("PostCreatePythonFailed", "Post Create Python Script failed. Check log for details.");
+			return false;
+		}
+		
+		return true;
+	}
 
 	TSharedPtr<IPlugin> CreateAndLoadNewPluginInternal(const FString& PluginName, const FString& NameToReplace, const FString& PluginLocation, const FString& PluginFilePath, const FPluginUtils::FNewPluginParamsWithDescriptor& CreationParams, FPluginUtils::FLoadPluginParams& LoadParams)
 	{
@@ -758,6 +938,35 @@ namespace PluginUtils
 				USourceControlHelpers::MarkFilesForAdd(NewFilePaths);
 				GWarn->EndSlowTask();
 			}
+			
+			if (!CreationParams.PostCreatePythonScriptPath.IsEmpty())
+			{
+				const bool bIsEnginePlugin = FPaths::IsUnderDirectory(PluginFilePath, FPaths::EnginePluginsDir());
+				
+				FText OutPythonScriptFailReason;
+				const bool bPythonOk = RunPostCreatePythonScript(
+					PluginName,
+					PluginFolder,
+					CreationParams.TemplateFolders,
+					CreationParams.PostCreatePythonScriptPath,
+					CreationParams.PostCreatePythonScriptArguments,
+					CreationParams.PostCreatePythonScriptAdditionalArguments,
+					bIsEnginePlugin,
+					OutPythonScriptFailReason
+				);
+
+				if (!bPythonOk)
+				{
+					if (LoadParams.OutFailReason)
+					{
+						*LoadParams.OutFailReason = OutPythonScriptFailReason;
+					}
+					
+					bSucceeded = false;
+					break;
+				}
+			}
+			
 		} while (false);
 
 		if (!CoreRedirectList.IsEmpty())
diff --git a/Engine/Plugins/Developer/PluginUtils/Source/PluginUtils/Public/PluginUtils.h b/Engine/Plugins/Developer/PluginUtils/Source/PluginUtils/Public/PluginUtils.h
--- a/Engine/Plugins/Developer/PluginUtils/Source/PluginUtils/Public/PluginUtils.h
+++ b/Engine/Plugins/Developer/PluginUtils/Source/PluginUtils/Public/PluginUtils.h
@@ -136,6 +136,14 @@ public:
 
 		/** If this plugin has Source, when should the module be loaded (may need to be earlier than default if used in blueprints) */
 		ELoadingPhase::Type LoadingPhase = ELoadingPhase::Default;
+		
+		/**
+		 * Strings used to run a python script once the plugin has been created.
+		 * If PostCreatePythonScriptPath is empty, it will not run anything.
+		 */
+		FString PostCreatePythonScriptPath;
+		FString PostCreatePythonScriptArguments;
+		FString PostCreatePythonScriptAdditionalArguments;
 	};
 
 	/**
@@ -160,6 +168,14 @@ public:
 		 * If empty, the plugin name will be used.
 		 */
 		FString NameToReplace;
+		
+		/**
+		 * Strings used to run a python script once the plugin has been created.
+		 * If PostCreatePythonScriptPath is empty, it will not run anything.
+		 */
+		FString PostCreatePythonScriptPath;
+		FString PostCreatePythonScriptArguments;
+		FString PostCreatePythonScriptAdditionalArguments;
 	};
 
 	/**
diff --git a/Engine/Plugins/Editor/PluginBrowser/Source/PluginBrowser/Private/NewPluginPythonOptionsData.cpp b/Engine/Plugins/Editor/PluginBrowser/Source/PluginBrowser/Private/NewPluginPythonOptionsData.cpp
new file mode 100644
--- /dev/null
+++ b/Engine/Plugins/Editor/PluginBrowser/Source/PluginBrowser/Private/NewPluginPythonOptionsData.cpp
@@ -0,0 +1,9 @@
+﻿// Copyright Epic Games, Inc. All Rights Reserved.
+
+
+#include "NewPluginPythonOptionsData.h"
+
+UNewPluginPythonOptionsData::UNewPluginPythonOptionsData(const FObjectInitializer& ObjectInitializer)
+	: Super(ObjectInitializer)
+{
+}
\ No newline at end of file
diff --git a/Engine/Plugins/Editor/PluginBrowser/Source/PluginBrowser/Private/NewPluginPythonOptionsData.h b/Engine/Plugins/Editor/PluginBrowser/Source/PluginBrowser/Private/NewPluginPythonOptionsData.h
new file mode 100644
--- /dev/null
+++ b/Engine/Plugins/Editor/PluginBrowser/Source/PluginBrowser/Private/NewPluginPythonOptionsData.h
@@ -0,0 +1,24 @@
+﻿// Copyright Epic Games, Inc. All Rights Reserved.
+
+#pragma once
+
+#include "CoreMinimal.h"
+#include "UObject/Object.h"
+#include "NewPluginPythonOptionsData.generated.h"
+
+/**
+ * An object used internally by the New Plugin Wizard to set the post-create python script options.
+ * This should not be used outside of the New Plugin Wizard dialog
+ */
+UCLASS()
+class UNewPluginPythonOptionsData : public UObject
+{
+	GENERATED_BODY()
+	
+public:
+	UNewPluginPythonOptionsData(const FObjectInitializer& ObjectInitializer);
+
+	/** Optional additional arguments for the python post-create script  */
+	UPROPERTY(EditAnywhere, Category="Post Create Python Script")
+	FString AdditionalArguments;
+};
diff --git a/Engine/Plugins/Editor/PluginBrowser/Source/PluginBrowser/Private/SNewPluginWizard.cpp b/Engine/Plugins/Editor/PluginBrowser/Source/PluginBrowser/Private/SNewPluginWizard.cpp
--- a/Engine/Plugins/Editor/PluginBrowser/Source/PluginBrowser/Private/SNewPluginWizard.cpp
+++ b/Engine/Plugins/Editor/PluginBrowser/Source/PluginBrowser/Private/SNewPluginWizard.cpp
@@ -34,6 +34,7 @@
 #include "DefaultPluginWizardDefinition.h"
 #include "PropertyEditorModule.h"
 #include "IDetailsView.h"
+#include "NewPluginPythonOptionsData.h"
 #include "Framework/Notifications/NotificationManager.h"
 #include "Widgets/Notifications/SNotificationList.h"
 #include "SourceCodeNavigation.h"
@@ -80,6 +81,9 @@ void SNewPluginWizard::Construct(const FArguments& Args, TSharedPtr<SDockTab> In
 		DetailsViewArgs.bHideSelectionTip = true;
 	}
 	TSharedPtr<IDetailsView> DescriptorDetailView = EditModule.CreateDetailView(DetailsViewArgs);
+	
+	PythonOptionsData = NewObject<UNewPluginPythonOptionsData>();
+	TSharedPtr<IDetailsView> PythonOptionsDetailView = EditModule.CreateDetailView(DetailsViewArgs);
 
 	if ( !PluginWizardDefinition.IsValid() )
 	{
@@ -179,6 +183,27 @@ void SNewPluginWizard::Construct(const FArguments& Args, TSharedPtr<SDockTab> In
 		];
 	}
 
+	// Add the python options data object if it exists
+	if (PythonOptionsData.IsValid() && PythonOptionsDetailView.IsValid())
+	{
+		PythonOptionsDetailView->SetObject(PythonOptionsData.Get());
+
+		MainContent->AddSlot()
+		.AutoHeight()
+		.Padding(PaddingAmount)
+		[
+			SNew(SBox)
+			.Visibility_Lambda([this]()
+			{
+				FPluginTemplateDescription* Template = PluginWizardDefinition->GetSelectedTemplate().Get();
+				return ((Template != nullptr) && !Template->PostCreatePythonScriptPath.IsEmpty()) ? EVisibility::Visible : EVisibility::Collapsed;
+			})
+			[
+				PythonOptionsDetailView.ToSharedRef()
+			]
+		];
+	}
+
 	MainContent->AddSlot()
 	.AutoHeight()
 	[
@@ -680,6 +705,14 @@ FReply SNewPluginWizard::OnCreatePluginClicked()
 		CreationParams.Descriptor.Description = DescriptorData->Description;
 		CreationParams.Descriptor.bIsBetaVersion = DescriptorData->bIsBetaVersion;
 	}
+	
+	CreationParams.PostCreatePythonScriptPath = Template->PostCreatePythonScriptPath;
+	CreationParams.PostCreatePythonScriptArguments = Template->PostCreatePythonScriptArguments;
+	
+	if (PythonOptionsData.IsValid())
+	{
+		CreationParams.PostCreatePythonScriptAdditionalArguments = PythonOptionsData->AdditionalArguments;
+	}
 
 	FText FailReason;
 	FPluginUtils::FLoadPluginParams LoadParams;
diff --git a/Engine/Plugins/Editor/PluginBrowser/Source/PluginBrowser/Private/SNewPluginWizard.h b/Engine/Plugins/Editor/PluginBrowser/Source/PluginBrowser/Private/SNewPluginWizard.h
--- a/Engine/Plugins/Editor/PluginBrowser/Source/PluginBrowser/Private/SNewPluginWizard.h
+++ b/Engine/Plugins/Editor/PluginBrowser/Source/PluginBrowser/Private/SNewPluginWizard.h
@@ -7,6 +7,7 @@
 class IPluginWizardDefinition;
 class SDockTab;
 class UNewPluginDescriptorData;
+class UNewPluginPythonOptionsData;
 namespace ESelectInfo { enum Type : int; }
 struct FPluginTemplateDescription;
 
@@ -128,6 +129,9 @@ private:
 private:
 	/** Additional user-defined descriptor data */
 	TWeakObjectPtr<UNewPluginDescriptorData> DescriptorData;
+	
+	/** Additional user-defined descriptor data */
+	TWeakObjectPtr<UNewPluginPythonOptionsData> PythonOptionsData;
 
 	/** The current plugin wizard definition */
 	TSharedPtr<IPluginWizardDefinition> PluginWizardDefinition;
diff --git a/Engine/Plugins/Runtime/GameFeatures/Source/GameFeaturesEditor/Private/GameFeaturePluginTemplate.cpp b/Engine/Plugins/Runtime/GameFeatures/Source/GameFeaturesEditor/Private/GameFeaturePluginTemplate.cpp
--- a/Engine/Plugins/Runtime/GameFeatures/Source/GameFeaturesEditor/Private/GameFeaturePluginTemplate.cpp
+++ b/Engine/Plugins/Runtime/GameFeatures/Source/GameFeaturesEditor/Private/GameFeaturePluginTemplate.cpp
@@ -16,8 +16,8 @@
 
 #define LOCTEXT_NAMESPACE "GameFeatures"
 
-FGameFeaturePluginTemplateDescription::FGameFeaturePluginTemplateDescription(FText InName, FText InDescription, FString InOnDiskPath, FString InDefaultSubfolder, FString InDefaultPluginName, TSubclassOf<UGameFeatureData> GameFeatureDataClassOverride, FString GameFeatureDataNameOverride, EPluginEnabledByDefault InEnabledByDefault)
-	: FPluginTemplateDescription(InName, InDescription, InOnDiskPath, /*bCanContainContent=*/ true, EHostType::Runtime)
+FGameFeaturePluginTemplateDescription::FGameFeaturePluginTemplateDescription(FText InName, FText InDescription, FString InOnDiskPath, FString InDefaultSubfolder, FString InDefaultPluginName, TSubclassOf<UGameFeatureData> GameFeatureDataClassOverride, FString GameFeatureDataNameOverride, EPluginEnabledByDefault InEnabledByDefault, FString InPostCreatePythonScriptPath, FString InPostCreatePythonScriptArguments)
+	: FPluginTemplateDescription(InName, InDescription, InOnDiskPath, /*bCanContainContent=*/ true, EHostType::Runtime, ELoadingPhase::Default, /*InSupportsContentOnlyProjects=*/ true, InPostCreatePythonScriptPath, InPostCreatePythonScriptArguments)
 {
 	SortPriority = 10;
 	bCanBePlacedInEngine = false;
diff --git a/Engine/Plugins/Runtime/GameFeatures/Source/GameFeaturesEditor/Private/GameFeaturesEditorModule.cpp b/Engine/Plugins/Runtime/GameFeatures/Source/GameFeaturesEditor/Private/GameFeaturesEditorModule.cpp
--- a/Engine/Plugins/Runtime/GameFeatures/Source/GameFeaturesEditor/Private/GameFeaturesEditorModule.cpp
+++ b/Engine/Plugins/Runtime/GameFeatures/Source/GameFeaturesEditor/Private/GameFeaturesEditorModule.cpp
@@ -298,7 +298,9 @@ class FGameFeaturesEditorModule : public FDefaultModuleImpl
 					PluginTemplate.DefaultPluginName,
 					PluginTemplate.DefaultGameFeatureDataClass,
 					PluginTemplate.DefaultGameFeatureDataName,
-					PluginTemplate.bIsEnabledByDefault ? EPluginEnabledByDefault::Enabled : EPluginEnabledByDefault::Disabled)));
+					PluginTemplate.bIsEnabledByDefault ? EPluginEnabledByDefault::Enabled : EPluginEnabledByDefault::Disabled,
+					PluginTemplate.PostCreatePythonScriptPath.FilePath,
+					PluginTemplate.PostCreatePythonScriptArguments)));
 			}
 		}
 	}
diff --git a/Engine/Plugins/Runtime/GameFeatures/Source/GameFeaturesEditor/Private/GameFeaturesEditorSettings.h b/Engine/Plugins/Runtime/GameFeatures/Source/GameFeaturesEditor/Private/GameFeaturesEditorSettings.h
--- a/Engine/Plugins/Runtime/GameFeatures/Source/GameFeaturesEditor/Private/GameFeaturesEditorSettings.h
+++ b/Engine/Plugins/Runtime/GameFeatures/Source/GameFeaturesEditor/Private/GameFeaturesEditorSettings.h
@@ -49,6 +49,14 @@ struct FPluginTemplateData
 	/** If true, the created plugin will be enabled by default without needing to be added to the project file. */
 	UPROPERTY(config, EditAnywhere, Category = Plugins)
 	bool bIsEnabledByDefault = false;
+	
+	/** Optional path to a post create python script that will run when new plugins are created from this template. */
+	UPROPERTY(config, EditAnywhere, Category = Plugins, meta = (RelativePath))
+	FFilePath PostCreatePythonScriptPath;
+	
+	/** Optional post create python script arguments. */
+	UPROPERTY(config, EditAnywhere, Category = Plugins)
+	FString PostCreatePythonScriptArguments;
 
 };
 
diff --git a/Engine/Plugins/Runtime/GameFeatures/Source/GameFeaturesEditor/Public/GameFeaturePluginTemplate.h b/Engine/Plugins/Runtime/GameFeatures/Source/GameFeaturesEditor/Public/GameFeaturePluginTemplate.h
--- a/Engine/Plugins/Runtime/GameFeatures/Source/GameFeaturesEditor/Public/GameFeaturePluginTemplate.h
+++ b/Engine/Plugins/Runtime/GameFeatures/Source/GameFeaturesEditor/Public/GameFeaturePluginTemplate.h
@@ -14,7 +14,8 @@
 struct FGameFeaturePluginTemplateDescription : public FPluginTemplateDescription
 {
 	UE_API FGameFeaturePluginTemplateDescription(FText InName, FText InDescription, FString InOnDiskPath, FString InDefaultSubfolder, FString InDefaultPluginName
-		, TSubclassOf<UGameFeatureData> GameFeatureDataClassOverride, FString GameFeatureDataNameOverride, EPluginEnabledByDefault InEnabledByDefault);
+		, TSubclassOf<UGameFeatureData> GameFeatureDataClassOverride, FString GameFeatureDataNameOverride, EPluginEnabledByDefault InEnabledByDefault
+		, FString InPostCreatePythonScriptPath, FString InPostCreatePythonScriptArguments);
 
 	UE_API virtual bool ValidatePathForPlugin(const FString& ProposedAbsolutePluginPath, FText& OutErrorMessage) override;
 	UE_API virtual void UpdatePathWhenTemplateSelected(FString& InOutPath) override;
diff --git a/Engine/Source/Editor/UnrealEd/Public/Features/IPluginsEditorFeature.h b/Engine/Source/Editor/UnrealEd/Public/Features/IPluginsEditorFeature.h
--- a/Engine/Source/Editor/UnrealEd/Public/Features/IPluginsEditorFeature.h
+++ b/Engine/Source/Editor/UnrealEd/Public/Features/IPluginsEditorFeature.h
@@ -46,6 +46,10 @@ struct FPluginTemplateDescription
 
 	/** Can the plugin be used in a content only project? */
 	bool bSupportsContentOnlyProjects;
+	
+	/** Python script to run after a plugin is created from this template */
+	FString PostCreatePythonScriptPath;
+	FString PostCreatePythonScriptArguments;
 
 	/** Called just before the plugin is created */
 	virtual void CustomizeDescriptorBeforeCreation(FPluginDescriptor& Descriptor) {}
@@ -73,7 +77,7 @@ struct FPluginTemplateDescription
 	virtual void UpdatePluginNameTextWhenTemplateUnselected(FText& OutPluginNameText) {}
 
 	/** Constructor */
-	FPluginTemplateDescription(FText InName, FText InDescription, FString InOnDiskPath, bool InCanContainContent, EHostType::Type InModuleDescriptorType, ELoadingPhase::Type InLoadingPhase = ELoadingPhase::Default, bool InSupportsContentOnlyProjects = true)
+	FPluginTemplateDescription(FText InName, FText InDescription, FString InOnDiskPath, bool InCanContainContent, EHostType::Type InModuleDescriptorType, ELoadingPhase::Type InLoadingPhase = ELoadingPhase::Default, bool InSupportsContentOnlyProjects = true, FString InPostCreatePythonScriptPath = FString(), FString InPostCreatePythonScriptArguments = FString())
 		: Name(InName)
 		, Description(InDescription)
 		, OnDiskPath(InOnDiskPath)
@@ -81,6 +85,8 @@ struct FPluginTemplateDescription
 		, ModuleDescriptorType(InModuleDescriptorType)
 		, LoadingPhase(InLoadingPhase)
 		, bSupportsContentOnlyProjects(InSupportsContentOnlyProjects)
+		, PostCreatePythonScriptPath(InPostCreatePythonScriptPath)
+		, PostCreatePythonScriptArguments(InPostCreatePythonScriptArguments)
 	{
 	}
 
