Call post-create python script on same editor instance

[RN] minor plugin

#Synced-CL 49987136

diff --git a/Engine/Plugins/Developer/PluginUtils/PluginUtils.uplugin b/Engine/Plugins/Developer/PluginUtils/PluginUtils.uplugin
--- a/Engine/Plugins/Developer/PluginUtils/PluginUtils.uplugin
+++ b/Engine/Plugins/Developer/PluginUtils/PluginUtils.uplugin
@@ -25,6 +25,12 @@
 		{
 			"Name": "EditorScriptingUtilities",
 			"Enabled": true
+		},
+		{
+			"Name": "PythonScriptPlugin",
+			"Enabled": false,
+			"Optional": true,
+			"TargetAllowList": [ "Editor" ]
 		}
 	]
 }
diff --git a/Engine/Plugins/Developer/PluginUtils/Source/PluginUtils/PluginUtils.Build.cs b/Engine/Plugins/Developer/PluginUtils/Source/PluginUtils/PluginUtils.Build.cs
--- a/Engine/Plugins/Developer/PluginUtils/Source/PluginUtils/PluginUtils.Build.cs
+++ b/Engine/Plugins/Developer/PluginUtils/Source/PluginUtils/PluginUtils.Build.cs
@@ -28,6 +28,16 @@ namespace UnrealBuildTool.Rules
 					"UnrealEd"
 				}
 			);
+			
+			if (Target.bBuildEditor && Target.bCompilePython)
+			{
+				PrivateDependencyModuleNames.Add("PythonScriptPlugin");
+				PublicDefinitions.Add("COMPILE_PYTHON=1");
+			}
+			else
+			{
+				PublicDefinitions.Add("COMPILE_PYTHON=0");
+			}
 		}
 	}
 }
diff --git a/Engine/Plugins/Developer/PluginUtils/Source/PluginUtils/Private/PluginUtils.cpp b/Engine/Plugins/Developer/PluginUtils/Source/PluginUtils/Private/PluginUtils.cpp
--- a/Engine/Plugins/Developer/PluginUtils/Source/PluginUtils/Private/PluginUtils.cpp
+++ b/Engine/Plugins/Developer/PluginUtils/Source/PluginUtils/Private/PluginUtils.cpp
@@ -19,14 +19,12 @@
 #include "Interfaces/IPluginManager.h"
 #include "Interfaces/IProjectManager.h"
 #include "Internationalization/TextPackageNamespaceUtil.h"
-#include "LocalizationDescriptor.h"
 #include "Misc/FeedbackContext.h"
 #include "Misc/FileHelper.h"
 #include "Misc/PackageName.h"
 #include "Misc/PathViews.h"
 #include "Misc/Paths.h"
 #include "PackageTools.h"
-#include "PluginReferenceDescriptor.h"
 #include "SourceControlHelpers.h"
 #include "SubobjectDataSubsystem.h"
 #include "UnrealEdMisc.h"
@@ -35,6 +33,10 @@
 #include "UObject/PropertyAccessUtil.h"
 #include "UObject/UObjectIterator.h"
 
+#if COMPILE_PYTHON
+#include "IPythonScriptPlugin.h"
+#endif
+
 DEFINE_LOG_CATEGORY_STATIC(LogPluginUtils, Log, All);
 
 #define LOCTEXT_NAMESPACE "PluginUtils"
@@ -609,13 +611,7 @@ namespace PluginUtils
 	
 	bool RunPostCreatePythonScript(const FString& PluginName, const FString& PluginFolder, const TArray<FString> TemplateFolders, const FString& ScriptPath, const FString& TemplateArgs, const FString& AdditionalArgs, const bool IsEnginePlugin, FText& OutFailReason)
 	{
-		FString AbsScriptPath = IFileManager::Get().ConvertToAbsolutePathForExternalAppForRead(*ScriptPath);
-		if (!FPaths::FileExists(AbsScriptPath))
-		{
-			OutFailReason = FText::Format(LOCTEXT("PostCreateScriptNotFound", "Post Create Python Script not found in: {0}"), FText::FromString(AbsScriptPath));
-			return false;
-		}
-		
+#if COMPILE_PYTHON
 		const TSharedPtr<IPlugin> PythonPlugin = IPluginManager::Get().FindPlugin(TEXT("PythonScriptPlugin"));
 		if (!PythonPlugin.IsValid() || !PythonPlugin->IsEnabled())
 		{
@@ -623,140 +619,185 @@ namespace PluginUtils
 			return false;
 		}
 		
-		const FString TemplatePluginNames = GetTemplatePluginNames(TemplateFolders);
-		if (TemplatePluginNames.IsEmpty())
+		IPythonScriptPlugin* PythonScriptPlugin = IPythonScriptPlugin::Get();
+		if (!PythonScriptPlugin)
 		{
-			OutFailReason = LOCTEXT("PostCreateNoTemplatePluginName", "Could not find a valid Template Plugin Name.");
+			OutFailReason = LOCTEXT("PostCreatePythonModuleNotAvailable", "PythonScriptPlugin module not found. Post Create Python Script cannot run without it.");
 			return false;
 		}
 		
-		const FString EditorExe = FUnrealEdMisc::Get().GetExecutableForCommandlets();
-		if (EditorExe.IsEmpty())
+		if (!PythonScriptPlugin->IsPythonAvailable())
 		{
-			OutFailReason = LOCTEXT("PostCreateNoEditorExe", "Could not find commandlet executable for running the post-create Python script.");
+			OutFailReason = LOCTEXT("PostCreatePythonNotAvailableOnTarget", "Engine does not have Python support enabled. Python is required to run Post Create python scripts.");
 			return false;
 		}
 		
-		void* ReadPipe = nullptr, *WritePipe = nullptr;
-		if (!FPlatformProcess::CreatePipe(ReadPipe, WritePipe))
+		if (!PythonScriptPlugin->IsPythonInitialized())
 		{
-			OutFailReason = LOCTEXT("PostCreatePipeFailed", "Failed to create pipe for post-create Python script.");
+			OutFailReason = LOCTEXT("PostCreatePythonNotInitialized", "PythonScriptPlugin loaded and enabled but not initialized. Wait until it is initialized and try again.");
 			return false;
 		}
-
-		FString CommandLine = FCommandLine::GetOriginal();
 		
-		auto WrapInQuotes = [](const FString& In) -> FString
-			{
-				FString Out = In;
-				Out.InsertAt(0, "\"");
-				Out.Append("\"");
-				return Out;
-			};
+		const FString AbsScriptPath = IFileManager::Get().ConvertToAbsolutePathForExternalAppForRead(*ScriptPath);
+		if (!FPaths::FileExists(AbsScriptPath))
+		{
+			OutFailReason = FText::Format(LOCTEXT("PostCreateScriptNotFound", "Post Create Python Script not found in: {0}"), FText::FromString(AbsScriptPath));
+			return false;
+		}
 		
-		auto SwapToForwardSlash = [](const FString& In) -> FString
-			{
-				FString Out = In;
-				Out.ReplaceInline(TEXT("\\"), TEXT("/")); // swap backslash to forward slash for paths
-				return Out;
-			};
+		const FString PluginRoot = IFileManager::Get().ConvertToAbsolutePathForExternalAppForRead(*PluginFolder);
+		if (!FPaths::DirectoryExists(PluginRoot))
+		{
+			OutFailReason = FText::Format(LOCTEXT("PostCreatePluginNotFound", "Plugin {0} not found in: {1}"), FText::FromString(PluginName), FText::FromString(PluginRoot));
+			return false;
+		}
 		
-		auto EscapeForScript = [](const FString& In) -> FString
+		auto ValidateArgsBlob = [](const FString& In, FText& OutError) -> bool
 			{
-				FString Out = In;
-				Out.ReplaceInline(TEXT("\'"), TEXT("\"")); // swap single for double quotes
-				Out.ReplaceInline(TEXT("\""), TEXT("\\\\\"\"\"")); // enough escapes to go through the commandlet + python command stripping 
-				return Out;
+				if (In.IsEmpty())
+				{
+					return true; // empty strings are valid
+				}
+				
+				// Unbalanced quotes
+				bool bInQuotes = false;
+				int32 BackslashCount = 0;
+				for (int32 i = 0; i < In.Len(); ++i)
+				{
+					const TCHAR Char = In[i];
+					
+					// Invalid characters
+					if (Char == '\0' || Char == '\r' || Char == '\n')
+					{
+						OutError = FText::Format(LOCTEXT("PostCreatePythonForbiddenCharsInArg", "Arguments passed in (\"{0}\") contain forbidden control characters."), FText::FromString(In));
+						return false;
+					}
+					
+					if (Char == '\\')
+					{
+						BackslashCount++;
+						continue;
+					}
+					
+					if (Char == '"')
+					{
+						if (BackslashCount % 2 == 0)
+						{
+							bInQuotes = !bInQuotes;
+						}
+					}
+					BackslashCount = 0;
+				}
+				
+				if (bInQuotes)
+				{
+					OutError = FText::Format(LOCTEXT("PostCreatePythonUnbalancedQuotes", "Arguments passed in (\"{0}\") contain unbalanced quotes."), FText::FromString(In));
+					return false;
+				}
+				
+				return true;
 			};
 		
-		FString TemplatePathArgs;
-		for (const FString& Folder : TemplateFolders)
+		if (!ValidateArgsBlob(TemplateArgs, OutFailReason) || !ValidateArgsBlob(AdditionalArgs, OutFailReason))
 		{
-			TemplatePathArgs += FString::Printf(TEXT("--TemplatePath \"%s\" "), *SwapToForwardSlash(FPaths::ConvertRelativePathToFull(Folder)));
+			return false;
 		}
 		
-		CommandLine += FString::Printf(TEXT(" -stdout -FullStdOutLogOutput -redirectstderr -EnablePlugins=\"%s\" -run=pythonscript -script=\"%s --PluginName %s --PluginRoot %s --IsEnginePlugin %d %s %s %s\""), *TemplatePluginNames, *EscapeForScript(AbsScriptPath), *EscapeForScript(WrapInQuotes(PluginName)), *EscapeForScript(WrapInQuotes(SwapToForwardSlash(PluginFolder))), IsEnginePlugin ? 1 : 0, *EscapeForScript(TemplatePathArgs), *EscapeForScript(TemplateArgs), *EscapeForScript(AdditionalArgs));
+		TArray<FString> Argv;
+		Argv.Add(AbsScriptPath);
 		
-		FScopedSlowTask Progress(/*InAmountOfWork=*/ 3.f, LOCTEXT("RunningPostCreatePython", "Running post-create Python script..."), /*bInEnabled=*/ true, *GWarn);
-		Progress.MakeDialog(/*bShowCancelButton=*/ false);
+		Argv.Add(TEXT("--PluginName"));
+		Argv.Add(PluginName);
 		
-		FProcHandle ProcHandle = FPlatformProcess::CreateProc(*EditorExe, *CommandLine, /*bLaunchDetached*/ false, /*bLaunchHidden*/ true, /*bLaunchReallyHidden*/ true, /*OutProcessId*/ nullptr, /*PriorityModifier*/ 0, /*OptionalWorkingDirectory*/ nullptr, /*PipeWrite*/ WritePipe);
-		if (!ProcHandle.IsValid())
-		{
-			FPlatformProcess::ClosePipe(ReadPipe, WritePipe);
-			OutFailReason = LOCTEXT( "PostCreateProcFailed", "Failed to create a process for the post-create Python script.");
-			return false;
-		}
+		Argv.Add(TEXT("--PluginRoot"));
+		Argv.Add(PluginRoot);
 		
-		UE_LOG(LogPluginUtils, Display, TEXT("Launching new process to run python post-create script..."));
+		Argv.Add(TEXT("--IsEnginePlugin"));
+		Argv.Add(IsEnginePlugin ? TEXT("1") : TEXT("0"));
 		
-		bool bSeenRunningPython = false;
-		bool bSeenScriptCompleted = false;
-		float ProgressThisFrame = 1.0f;
-		FText CurrentStatusText = LOCTEXT( "PostCreateLaunchProc", "Launching new process to run python post-create script...");
-		
-		int32 ReturnCode = -1;
-		while (FPlatformProcess::IsProcRunning(ProcHandle))
+		for (const FString& Folder : TemplateFolders)
 		{
-			FPlatformProcess::Sleep(0.1f);
-			
-			const FString STDLog = FPlatformProcess::ReadPipe(ReadPipe);
-			if (!STDLog.IsEmpty())
+			Argv.Add(TEXT("--TemplatePath"));
+			Argv.Add(FPaths::ConvertRelativePathToFull(Folder));
+		}
+		
+		auto TokenizeArgsString = [](const FString& In, TArray<FString>& OutArgv) -> void
 			{
-				// Log the process' output to our output log
-				UE_LOG(LogPluginUtils, Display, TEXT("%s"), *STDLog);
+				if (In.IsEmpty())
+				{
+					return; // nothing to do
+				}
 				
-				// Update UI Progress bar based on log entries
-				TArray<FString> Lines;
-				STDLog.ParseIntoArrayLines(Lines);
-
-				for (const FString& Line : Lines)
+				const TCHAR* Stream = *In;
+				FString Arg;
+				while(FParse::Token(Stream, Arg, /*bUseEscape=*/ true))
 				{
-					if (!bSeenRunningPython &&	
-						Line.Contains(TEXT("LogPythonScriptCommandlet: Display: Running Python script")))
-					{
-						bSeenRunningPython = true;
-						ProgressThisFrame = 1.0f;
-						CurrentStatusText = LOCTEXT( "PostCreateRunningScript", "Running Python script in process...");
-					}
-					else if (!bSeenScriptCompleted &&
-						 (Line.Contains(TEXT("LogPythonScriptCommandlet: Display: Python script executed successfully")) ||
-						  Line.Contains(TEXT("Commandlet PythonScriptCommandlet_0 finished execution")) ||
-						  Line.Contains(TEXT("Success - 0 error(s), 0 warning(s)")) ||
-						  Line.Contains(TEXT("Execution of commandlet took"))))
-					{
-						bSeenScriptCompleted = true;
-						ProgressThisFrame = 1.0f;
-						CurrentStatusText = LOCTEXT( "PostCreateScriptFinished", "Python script finished. Closing process...");
-					}
+					OutArgv.Add(Arg);
 				}
-			}
-
-			// Update UI progress bar
-			Progress.EnterProgressFrame(ProgressThisFrame, CurrentStatusText);
-			ProgressThisFrame = 0;
+			};
+			
+		TokenizeArgsString(TemplateArgs, Argv);
+		TokenizeArgsString(AdditionalArgs, Argv);
+		
+		auto EscapeAndQuote = [](const FString& In) -> FString
+			{
+				const bool bAlreadyQuoted = In.Len() >= 2 && In.StartsWith(TEXT("\"")) && In.EndsWith(TEXT("\""));
+				FString InternalString = bAlreadyQuoted ? In.Mid(1, In.Len() - 2) : In;
+				
+				// Escape content
+				InternalString = InternalString.Replace(TEXT("\\"), TEXT("\\\\"));
+				InternalString = InternalString.Replace(TEXT("\""), TEXT("\\\""));
+				
+				const bool bNeedsQuotes = InternalString.Contains(TEXT("\"")) || InternalString.Contains(TEXT(" ")) || InternalString.Contains(TEXT("\t"));				
+				return bNeedsQuotes ? FString::Printf(TEXT("\"%s\""), *InternalString) : InternalString;
+			};
+		
+		FString CommandLine;
+		for (const FString& Arg : Argv)
+		{
+			CommandLine += EscapeAndQuote(Arg) + TEXT(" ");
 		}
 		
-		UE_LOG(LogPluginUtils, Display, TEXT("Python post-create script process completed."));
+		FPythonCommandEx PythonCommand;
+		PythonCommand.ExecutionMode = EPythonCommandExecutionMode::ExecuteFile;
+		PythonCommand.Flags |= EPythonCommandFlags::Unattended;
+		PythonCommand.Command = CommandLine;
+		
+		FScopedSlowTask Progress(/*InAmountOfWork=*/ 1.f, LOCTEXT("RunningPostCreatePython", "Running post-create Python script..."), /*bInEnabled=*/ true, *GWarn);
+		Progress.MakeDialog(/*bShowCancelButton=*/ false);
+		Progress.EnterProgressFrame(1.f);
 		
-		FPlatformProcess::ClosePipe(ReadPipe, WritePipe);
+		UE_LOG(LogPluginUtils, Display, TEXT("Running python post-create script..."));
+		const bool bSuccess = PythonScriptPlugin->ExecPythonCommandEx(PythonCommand);
+		UE_LOG(LogPluginUtils, Display, TEXT("Python post-create script process completed."));
 		
-		if (!FPlatformProcess::GetProcReturnCode(ProcHandle, &ReturnCode))
+		for (const FPythonLogOutputEntry& LogOutputEntry : PythonCommand.LogOutput)
 		{
-			FPlatformProcess::CloseProc(ProcHandle);
-			OutFailReason = LOCTEXT("PostCreateFailedGetReturnCode", "Failed to get return code for post-create Python script process.");
-			return false;
+			switch (LogOutputEntry.Type)
+			{
+			case EPythonLogOutputType::Error:
+				UE_LOG(LogPluginUtils, Error, TEXT("[PythonPostCreate] %s"), *LogOutputEntry.Output);
+				break;
+			case EPythonLogOutputType::Warning:
+				UE_LOG(LogPluginUtils, Warning, TEXT("[PythonPostCreate] %s"), *LogOutputEntry.Output);
+				break;
+			case EPythonLogOutputType::Info:
+				UE_LOG(LogPluginUtils, Display, TEXT("[PythonPostCreate] %s"), *LogOutputEntry.Output);
+				break;
+			}
 		}
 		
-		FPlatformProcess::CloseProc(ProcHandle);
-		
-		if (ReturnCode != 0)
+		if (!bSuccess)
 		{
 			OutFailReason = LOCTEXT("PostCreatePythonFailed", "Post Create Python Script failed. Check log for details.");
 			return false;
 		}
 		
 		return true;
+#else
+		OutFailReason = LOCTEXT("PostCreatePythonNotAvailableOnTarget", "Engine does not have Python support enabled. Python is required to run Post Create python scripts.");
+		return false;
+#endif
 	}
 
 	TSharedPtr<IPlugin> CreateAndLoadNewPluginInternal(const FString& PluginName, const FString& NameToReplace, const FString& PluginLocation, const FString& PluginFilePath, const FPluginUtils::FNewPluginParamsWithDescriptor& CreationParams, FPluginUtils::FLoadPluginParams& LoadParams)
