Move Python Post Creation Script Run to after removal of package redirectors

[RN] minor plugin

#Synced-CL 50374992

diff --git a/Engine/Plugins/Developer/PluginUtils/Source/PluginUtils/Private/PluginUtils.cpp b/Engine/Plugins/Developer/PluginUtils/Source/PluginUtils/Private/PluginUtils.cpp
--- a/Engine/Plugins/Developer/PluginUtils/Source/PluginUtils/Private/PluginUtils.cpp
+++ b/Engine/Plugins/Developer/PluginUtils/Source/PluginUtils/Private/PluginUtils.cpp
@@ -980,34 +980,6 @@ namespace PluginUtils
 				GWarn->EndSlowTask();
 			}
 			
-			if (!CreationParams.PostCreatePythonScriptPath.IsEmpty())
-			{
-				const bool bIsEnginePlugin = FPaths::IsUnderDirectory(PluginFilePath, FPaths::EnginePluginsDir());
-				
-				FText OutPythonScriptFailReason;
-				const bool bPythonOk = RunPostCreatePythonScript(
-					PluginName,
-					PluginFolder,
-					CreationParams.TemplateFolders,
-					CreationParams.PostCreatePythonScriptPath,
-					CreationParams.PostCreatePythonScriptArguments,
-					CreationParams.PostCreatePythonScriptAdditionalArguments,
-					bIsEnginePlugin,
-					OutPythonScriptFailReason
-				);
-
-				if (!bPythonOk)
-				{
-					if (LoadParams.OutFailReason)
-					{
-						*LoadParams.OutFailReason = OutPythonScriptFailReason;
-					}
-					
-					bSucceeded = false;
-					break;
-				}
-			}
-			
 		} while (false);
 
 		if (!CoreRedirectList.IsEmpty())
@@ -1015,6 +987,33 @@ namespace PluginUtils
 			// Remove the package redirectors now that all the generated content has been saved.
 			FCoreRedirects::RemoveRedirectList(CoreRedirectList, TEXT("GenerateNewPlugin"));
 		}
+			
+		if (!CreationParams.PostCreatePythonScriptPath.IsEmpty())
+		{
+			const bool bIsEnginePlugin = FPaths::IsUnderDirectory(PluginFilePath, FPaths::EnginePluginsDir());
+				
+			FText OutPythonScriptFailReason;
+			const bool bPythonOk = RunPostCreatePythonScript(
+				PluginName,
+				PluginFolder,
+				CreationParams.TemplateFolders,
+				CreationParams.PostCreatePythonScriptPath,
+				CreationParams.PostCreatePythonScriptArguments,
+				CreationParams.PostCreatePythonScriptAdditionalArguments,
+				bIsEnginePlugin,
+				OutPythonScriptFailReason
+			);
+
+			if (!bPythonOk)
+			{
+				if (LoadParams.OutFailReason)
+				{
+					*LoadParams.OutFailReason = OutPythonScriptFailReason;
+				}
+					
+				bSucceeded = false;
+			}
+		}
 
 		if (!bSucceeded)
 		{
